HairTools: move hair tools from MeshModelingToolset plugin to new experimental HairModelingToolset plugin that is disabled by default. Use Modeling Mode Extension API to add these Tools to Modeling Mode when the plugin is enabled.
#rnx

diff --git a/Engine/Plugins/Editor/ModelingToolsEditorMode/Source/ModelingToolsEditorMode/Private/ModelingToolsEditorMode.cpp b/Engine/Plugins/Editor/ModelingToolsEditorMode/Source/ModelingToolsEditorMode/Private/ModelingToolsEditorMode.cpp
--- a/Engine/Plugins/Editor/ModelingToolsEditorMode/Source/ModelingToolsEditorMode/Private/ModelingToolsEditorMode.cpp
+++ b/Engine/Plugins/Editor/ModelingToolsEditorMode/Source/ModelingToolsEditorMode/Private/ModelingToolsEditorMode.cpp
@@ -97,11 +97,6 @@
 #include "Physics/ExtractCollisionGeometryTool.h"
 //#include "Physics/EditCollisionGeometryTool.h"
 
-// hair tools
-#include "Hair/GroomToMeshTool.h"
-#include "Hair/GroomCardsEditorTool.h"
-#include "GenerateLODMeshesTool.h"
-
 // asset tools
 #include "Tools/GenerateStaticMeshLODAssetTool.h"
 #include "Tools/LODManagerTool.h"
@@ -584,18 +579,6 @@ void UModelingToolsEditorMode::Enter()
 	auto ExtractCollisionGeoToolBuilder = NewObject<UExtractCollisionGeometryToolBuilder>();
 	RegisterTool(ToolManagerCommands.BeginExtractCollisionGeometryTool, TEXT("BeginExtractCollisionGeometryTool"), ExtractCollisionGeoToolBuilder);
 
-
-
-	// (experimental) hair tools
-
-	UGroomToMeshToolBuilder* GroomToMeshToolBuilder = NewObject<UGroomToMeshToolBuilder>();
-	RegisterTool(ToolManagerCommands.BeginGroomToMeshTool, TEXT("BeginGroomToMeshTool"), GroomToMeshToolBuilder);
-
-	RegisterTool(ToolManagerCommands.BeginGroomCardsEditorTool, TEXT("BeginGroomCardsEditorTool"), NewObject<UGroomCardsEditorToolBuilder>());
-
-	UGenerateLODMeshesToolBuilder* GenerateLODMeshesToolBuilder = NewObject<UGenerateLODMeshesToolBuilder>();
-	RegisterTool(ToolManagerCommands.BeginGenerateLODMeshesTool, TEXT("BeginGenerateLODMeshesTool"), GenerateLODMeshesToolBuilder);
-
 	// PolyModeling tools
 	auto RegisterPolyModelSelectTool = [&](EEditMeshPolygonsToolSelectionMode SelectionMode, TSharedPtr<FUICommandInfo> UICommand, FString StringName)
 	{
diff --git a/Engine/Plugins/Editor/ModelingToolsEditorMode/Source/ModelingToolsEditorMode/Private/ModelingToolsEditorModeToolkit.cpp b/Engine/Plugins/Editor/ModelingToolsEditorMode/Source/ModelingToolsEditorMode/Private/ModelingToolsEditorModeToolkit.cpp
--- a/Engine/Plugins/Editor/ModelingToolsEditorMode/Source/ModelingToolsEditorMode/Private/ModelingToolsEditorModeToolkit.cpp
+++ b/Engine/Plugins/Editor/ModelingToolsEditorMode/Source/ModelingToolsEditorMode/Private/ModelingToolsEditorModeToolkit.cpp
@@ -457,13 +457,12 @@ static const FName TransformTabName(TEXT("Transform"));
 static const FName DeformTabName(TEXT("Deform"));
 static const FName VolumesTabName(TEXT("Volumes"));
 static const FName PrototypesTabName(TEXT("Prototypes"));
-static const FName HairTabName(TEXT("Hair"));
 static const FName PolyEditTabName(TEXT("PolyEdit"));
 static const FName VoxToolsTabName(TEXT("VoxOps"));
 static const FName LODToolsTabName(TEXT("LODs"));
 
 
-const TArray<FName> FModelingToolsEditorModeToolkit::PaletteNames_Standard = { PrimitiveTabName, CreateTabName, PolyModelingTabName, TriModelingTabName, DeformTabName, TransformTabName, MeshProcessingTabName, VoxToolsTabName, AttributesTabName, UVTabName, VolumesTabName, LODToolsTabName, HairTabName };
+const TArray<FName> FModelingToolsEditorModeToolkit::PaletteNames_Standard = { PrimitiveTabName, CreateTabName, PolyModelingTabName, TriModelingTabName, DeformTabName, TransformTabName, MeshProcessingTabName, VoxToolsTabName, AttributesTabName, UVTabName, VolumesTabName, LODToolsTabName };
 
 
 void FModelingToolsEditorModeToolkit::GetToolPaletteNames(TArray<FName>& PaletteNames) const
@@ -655,12 +654,6 @@ void FModelingToolsEditorModeToolkit::BuildToolPalette_Experimental(FName Palett
 		//ToolbarBuilder.AddToolBarButton(Commands.BeginEditCollisionGeometryTool);
 		ToolbarBuilder.AddToolBarButton(Commands.BeginExtractCollisionGeometryTool);
 	}
-	else if (PaletteIndex == HairTabName)
-	{
-		ToolbarBuilder.AddToolBarButton(Commands.BeginGroomToMeshTool);
-		ToolbarBuilder.AddToolBarButton(Commands.BeginGroomCardsEditorTool);
-		ToolbarBuilder.AddToolBarButton(Commands.BeginGenerateLODMeshesTool);
-	}
 	else if (PaletteIndex == PrototypesTabName)
 	{
 		ToolbarBuilder.AddToolBarButton(Commands.BeginAddPatchTool);
diff --git a/Engine/Plugins/Editor/ModelingToolsEditorMode/Source/ModelingToolsEditorMode/Private/ModelingToolsManagerActions.cpp b/Engine/Plugins/Editor/ModelingToolsEditorMode/Source/ModelingToolsEditorMode/Private/ModelingToolsManagerActions.cpp
--- a/Engine/Plugins/Editor/ModelingToolsEditorMode/Source/ModelingToolsEditorMode/Private/ModelingToolsManagerActions.cpp
+++ b/Engine/Plugins/Editor/ModelingToolsEditorMode/Source/ModelingToolsEditorMode/Private/ModelingToolsManagerActions.cpp
@@ -109,13 +109,7 @@ void FModelingToolsManagerCommands::RegisterCommands()
 	UI_COMMAND(BeginSkinWeightsPaintTool, "SkinWts", "Start the Paint Skin Weights Tool", EUserInterfaceActionType::Button, FInputChord());
 	UI_COMMAND(BeginSkinWeightsBindingTool, "SkinBind", "Start the Skin Weights Binding Tool", EUserInterfaceActionType::Button, FInputChord());
 
-	UI_COMMAND(BeginGroomToMeshTool, "Helmet", "Generate Helmet Mesh for Selected Groom", EUserInterfaceActionType::ToggleButton, FInputChord());
-	UI_COMMAND(BeginGenerateLODMeshesTool, "HlmLOD", "Generate LODS for Hair Helmet", EUserInterfaceActionType::ToggleButton, FInputChord());
-	UI_COMMAND(BeginGroomCardsToMeshTool, "CardsToMesh", "Hair Cards to Mesh Tool", EUserInterfaceActionType::ToggleButton, FInputChord());
-	UI_COMMAND(BeginGroomMeshToCardsTool, "MeshToCards", "Mesh to Hair Cards Tool", EUserInterfaceActionType::ToggleButton, FInputChord());
-	UI_COMMAND(BeginGroomCardsEditorTool, "CardsEd", "Edit Hair Cards", EUserInterfaceActionType::ToggleButton, FInputChord());
 	UI_COMMAND(BeginSubdividePolyTool, "SubD", "Subdivide Mesh Polygroups or Triangles", EUserInterfaceActionType::ToggleButton, FInputChord());
-
 	
 	UI_COMMAND(BeginLODManagerTool, "LODMgr", "Static Mesh Asset LOD Manager", EUserInterfaceActionType::ToggleButton, FInputChord());
 	UI_COMMAND(BeginGenerateStaticMeshLODAssetTool, "AutoLOD", "Generate Static Mesh LOD Asset", EUserInterfaceActionType::ToggleButton, FInputChord());
diff --git a/Engine/Plugins/Editor/ModelingToolsEditorMode/Source/ModelingToolsEditorMode/Public/ModelingToolsManagerActions.h b/Engine/Plugins/Editor/ModelingToolsEditorMode/Source/ModelingToolsEditorMode/Public/ModelingToolsManagerActions.h
--- a/Engine/Plugins/Editor/ModelingToolsEditorMode/Source/ModelingToolsEditorMode/Public/ModelingToolsManagerActions.h
+++ b/Engine/Plugins/Editor/ModelingToolsEditorMode/Source/ModelingToolsEditorMode/Public/ModelingToolsManagerActions.h
@@ -104,12 +104,6 @@ public:
 	TSharedPtr<FUICommandInfo> BeginSkinWeightsPaintTool;
 	TSharedPtr<FUICommandInfo> BeginSkinWeightsBindingTool;
 	
-	TSharedPtr<FUICommandInfo> BeginGroomToMeshTool;
-	TSharedPtr<FUICommandInfo> BeginGenerateLODMeshesTool;
-	TSharedPtr<FUICommandInfo> BeginGroomCardsToMeshTool;
-	TSharedPtr<FUICommandInfo> BeginGroomMeshToCardsTool;
-	TSharedPtr<FUICommandInfo> BeginGroomCardsEditorTool;
-
 	TSharedPtr<FUICommandInfo> BeginLODManagerTool;
 	TSharedPtr<FUICommandInfo> BeginGenerateStaticMeshLODAssetTool;
 
diff --git a/Engine/Plugins/Experimental/HairModelingToolset/HairModelingToolset.uplugin b/Engine/Plugins/Experimental/HairModelingToolset/HairModelingToolset.uplugin
new file mode 100644
--- /dev/null
+++ b/Engine/Plugins/Experimental/HairModelingToolset/HairModelingToolset.uplugin
@@ -0,0 +1,40 @@
+{
+	"FileVersion": 3,
+	"Version": 1,
+	"VersionName": "0.1",
+	"FriendlyName": "Hair Modeling Toolset",
+	"Description": "A set of Interactive Tools for Groom/Hair Modeling",
+	"Category": "Other",
+	"CreatedBy": "Epic Games, Inc.",
+	"CreatedByURL": "http://epicgames.com",
+	"DocsURL": "",
+	"MarketplaceURL": "",
+	"SupportURL": "",
+	"CanContainContent": true,
+	"IsExperimentalVersion": true,
+	"Installed": false,
+	"Hidden": false,
+	"Modules": [
+		{
+			"Name": "HairModelingToolset",
+			"Type": "Editor",
+			"LoadingPhase": "Default",
+			"WhitelistPlatforms": [ "Win64" ]
+		},
+	],
+	"Plugins": [
+		{
+			"Name": "GeometryProcessing",
+			"Enabled": true
+		},
+		{
+			"Name": "MeshModelingToolset",
+			"Enabled": true,
+			"WhitelistTargets": [ "Editor" ]
+		},
+		{
+			"Name": "ModelingToolsEditorMode",
+			"Enabled": true
+		}		
+	]
+}
diff --git a/Engine/Plugins/Experimental/HairModelingToolset/Source/HairModelingToolset/HairModelingToolset.Build.cs b/Engine/Plugins/Experimental/HairModelingToolset/Source/HairModelingToolset/HairModelingToolset.Build.cs
new file mode 100644
--- /dev/null
+++ b/Engine/Plugins/Experimental/HairModelingToolset/Source/HairModelingToolset/HairModelingToolset.Build.cs
@@ -0,0 +1,83 @@
+// Copyright Epic Games, Inc. All Rights Reserved.
+
+using UnrealBuildTool;
+
+public class HairModelingToolset : ModuleRules
+{
+	public HairModelingToolset(ReadOnlyTargetRules Target) : base(Target)
+	{
+		PCHUsage = ModuleRules.PCHUsageMode.UseExplicitOrSharedPCHs;
+		
+		PublicIncludePaths.AddRange(
+			new string[] {
+				// ... add public include paths required here ...
+			}
+			);
+				
+		
+		PrivateIncludePaths.AddRange(
+			new string[] {
+				// ... add other private include paths required here ...
+			}
+			);
+			
+		
+		PublicDependencyModuleNames.AddRange(
+			new string[]
+			{
+				"Core",
+				"Eigen",
+                "InteractiveToolsFramework",
+				"GeometryCore",
+				"GeometryFramework",
+				"GeometryAlgorithms",
+				"DynamicMesh",
+				"MeshConversion",
+				"MeshDescription",
+                "StaticMeshDescription",
+				"ModelingComponents",
+				"ModelingOperators",
+				"ModelingOperatorsEditorOnly",
+				"MeshModelingToolsExp",
+				"HairStrandsCore",
+				"ModelingToolsEditorMode"
+
+				// ... add other public dependencies that you statically link with here ...
+			}
+            );
+			
+		
+		PrivateDependencyModuleNames.AddRange(
+			new string[]
+			{
+				"CoreUObject",
+				"Engine",
+				"InputCore",
+				"PhysicsCore",
+				"RenderCore",
+				"RHI",
+
+				"ApplicationCore",
+				"Slate",
+				"SlateCore",
+				"UnrealEd",
+				"EditorFramework",
+				"ContentBrowser",
+				"LevelEditor",
+				"StatusBar",
+				"EditorStyle",
+				"Projects"
+
+				// ... add private dependencies that you statically link with here ...
+			}
+			);
+		
+		
+		DynamicallyLoadedModuleNames.AddRange(
+			new string[]
+			{
+				// ... add any modules that your module loads dynamically here ...
+			}
+			);
+	}
+}
diff --git a/Engine/Plugins/Experimental/HairModelingToolset/Source/HairModelingToolset/Private/GenerateLODMeshesTool.cpp b/Engine/Plugins/Experimental/HairModelingToolset/Source/HairModelingToolset/Private/GenerateLODMeshesTool.cpp
new file mode 100644
--- /dev/null
+++ b/Engine/Plugins/Experimental/HairModelingToolset/Source/HairModelingToolset/Private/GenerateLODMeshesTool.cpp
@@ -0,0 +1,416 @@
+// Copyright Epic Games, Inc. All Rights Reserved.
+
+#include "GenerateLODMeshesTool.h"
+#include "InteractiveToolManager.h"
+#include "Properties/RemeshProperties.h"
+#include "ToolBuilderUtil.h"
+
+#include "ToolSetupUtil.h"
+#include "Util/ColorConstants.h"
+
+#include "DynamicMesh/DynamicMesh3.h"
+
+#include "MeshDescriptionToDynamicMesh.h"
+#include "DynamicMeshToMeshDescription.h"
+
+#include "ModelingObjectsCreationAPI.h"
+#include "Selection/ToolSelectionUtil.h"
+
+#include "SceneManagement.h" // for FPrimitiveDrawInterface
+
+#include "Modules/ModuleManager.h"
+
+#include "TargetInterfaces/MaterialProvider.h"
+#include "TargetInterfaces/MeshDescriptionProvider.h"
+#include "TargetInterfaces/PrimitiveComponentBackedTarget.h"
+
+#include "IMeshReductionManagerModule.h"
+#include "IMeshReductionInterfaces.h"
+
+
+#if WITH_EDITOR
+#include "Misc/ScopedSlowTask.h"
+#endif
+
+#include "ExplicitUseGeometryMathTypes.h"		// using UE::Geometry::(math types)
+using namespace UE::Geometry;
+
+#define LOCTEXT_NAMESPACE "UGenerateLODMeshesTool"
+
+/*
+ * ToolBuilder
+ */
+USingleSelectionMeshEditingTool* UGenerateLODMeshesToolBuilder::CreateNewTool(const FToolBuilderState& SceneState) const
+{
+	return NewObject<UGenerateLODMeshesTool>(SceneState.ToolManager);
+}
+
+/*
+ * Tool
+ */
+UGenerateLODMeshesToolProperties::UGenerateLODMeshesToolProperties()
+{
+	GroupBoundaryConstraint = EGroupBoundaryConstraint::Ignore;
+	MaterialBoundaryConstraint = EMaterialBoundaryConstraint::Ignore;
+
+	// hardcoded for hair helmet
+
+	FLODLevelGenerateSettings LOD0;
+	LOD0.SimplifierType = ESimplifyType::UEStandard;
+	LOD0.TargetMode = ESimplifyTargetType::VertexCount;
+	LOD0.TargetPercentage = 50;
+	LOD0.TargetCount = 500;
+	LODLevels.Add(LOD0);
+
+	FLODLevelGenerateSettings LOD1;
+	LOD1.SimplifierType = ESimplifyType::UEStandard;
+	LOD1.TargetMode = ESimplifyTargetType::VertexCount;
+	LOD1.TargetPercentage = 30;
+	LOD1.TargetCount = 250;
+	LODLevels.Add(LOD1);
+
+	FLODLevelGenerateSettings LOD2;
+	LOD2.SimplifierType = ESimplifyType::UEStandard;
+	LOD2.TargetMode = ESimplifyTargetType::VertexCount;
+	LOD2.TargetPercentage = 15;
+	LOD2.TargetCount = 150;
+	LODLevels.Add(LOD2);
+}
+
+void UGenerateLODMeshesTool::Setup()
+{
+	UInteractiveTool::Setup();
+
+	SetToolDisplayName(LOCTEXT("ToolName", "Generate HairHelmet LODs"));
+
+	{
+		// if in editor, create progress indicator dialog because building mesh copies can be slow (for very large meshes)
+		// this is especially needed because of the copy we make of the meshdescription; for Reasons, copying meshdescription is pretty slow
+#if WITH_EDITOR
+		static const FText SlowTaskText = LOCTEXT("GenerateLODMeshesInit", "Building mesh simplification data...");
+
+		FScopedSlowTask SlowTask(3.0f, SlowTaskText);
+		SlowTask.MakeDialog();
+
+		// Declare progress shortcut lambdas
+		auto EnterProgressFrame = [&SlowTask](int Progress)
+	{
+			SlowTask.EnterProgressFrame((float)Progress);
+		};
+#else
+		auto EnterProgressFrame = [](int Progress) {};
+#endif
+		IMeshDescriptionProvider* TargetMeshProvider = Cast<IMeshDescriptionProvider>(Target);
+		OriginalMeshDescription = MakeShared<FMeshDescription, ESPMode::ThreadSafe>(*TargetMeshProvider->GetMeshDescription());
+		EnterProgressFrame(1);
+		// aux-data isn't deep copied - by default it is built during initial evaluation (not thread safe)  
+		// so force aux-data rebuild now before multiple UE4 simplifiers try to use it in parallel.
+		OriginalMeshDescription->BuildIndexers();   
+		OriginalMesh = MakeShared<FDynamicMesh3, ESPMode::ThreadSafe>();
+		FMeshDescriptionToDynamicMesh Converter;
+		Converter.Convert(TargetMeshProvider->GetMeshDescription(), *OriginalMesh);
+		EnterProgressFrame(2);
+		OriginalMeshSpatial = MakeShared<FDynamicMeshAABBTree3, ESPMode::ThreadSafe>(OriginalMesh.Get(), true);
+	}
+
+	WorldBounds = FAxisAlignedBox3d(Cast<IPrimitiveComponentBackedTarget>(Target)->GetOwnerActor()->GetComponentsBoundingBox());
+
+	// initialize our properties
+	SimplifyProperties = NewObject<UGenerateLODMeshesToolProperties>(this);
+	SimplifyProperties->RestoreProperties(this);
+	AddToolPropertySource(SimplifyProperties);
+
+	SimplifyProperties->WatchProperty(SimplifyProperties->TargetMode, [this](ESimplifyTargetType) { InvalidateAllPreviews(); });
+	SimplifyProperties->WatchProperty(SimplifyProperties->SimplifierType, [this](ESimplifyType) { InvalidateAllPreviews(); });
+	SimplifyProperties->WatchProperty(SimplifyProperties->TargetPercentage, [this](int32) { InvalidateAllPreviews(); });
+	SimplifyProperties->WatchProperty(SimplifyProperties->TargetEdgeLength, [this](int32) { InvalidateAllPreviews(); });
+	SimplifyProperties->WatchProperty(SimplifyProperties->TargetCount, [this](int32) { InvalidateAllPreviews(); });
+	SimplifyProperties->WatchProperty(SimplifyProperties->bDiscardAttributes, [this](bool) { InvalidateAllPreviews(); });
+	SimplifyProperties->WatchProperty(SimplifyProperties->bReproject, [this](bool) { InvalidateAllPreviews(); });
+
+	SimplifyProperties->WatchProperty(SimplifyProperties->bShowGroupColors,
+									  [this](bool bNewValue) { UpdateVisualization(); });
+	SimplifyProperties->WatchProperty(SimplifyProperties->bShowWireframe,
+									  [this](bool bNewValue) { UpdateVisualization(); });
+
+	UpdateNumPreviews();
+
+	UpdateVisualization();
+}
+
+
+void UGenerateLODMeshesTool::Shutdown(EToolShutdownType ShutdownType)
+{
+	SimplifyProperties->SaveProperties(this);
+	Cast<IPrimitiveComponentBackedTarget>(Target)->SetOwnerVisibility(true);
+
+	if (ShutdownType == EToolShutdownType::Accept)
+	{
+		GenerateAssets();
+	}
+	else
+	{
+		for (UMeshOpPreviewWithBackgroundCompute* Preview : Previews)
+		{
+			Preview->Shutdown();
+		}
+	}
+}
+
+
+void UGenerateLODMeshesTool::OnTick(float DeltaTime)
+{
+	// invalidate any previews that have invalid caches
+	for (int32 k = 0; k < CachedLODLevels.Num() && k < SimplifyProperties->LODLevels.Num(); ++k)
+	{
+		if (CachedLODLevels[k] != SimplifyProperties->LODLevels[k])
+		{
+			CachedLODLevels[k] = SimplifyProperties->LODLevels[k];
+			PreviewFactories[k]->LODSettings = SimplifyProperties->LODLevels[k];
+			Previews[k]->InvalidateResult();
+		}
+	}
+
+
+	for (UMeshOpPreviewWithBackgroundCompute* Preview : Previews)
+	{
+		Preview->Tick(DeltaTime);
+	}
+}
+
+
+void UGenerateLODMeshesTool::UpdateNumPreviews()
+{
+	int32 NumPreviews = SimplifyProperties->LODLevels.Num();
+	int32 CurNumPreviews = Previews.Num();
+
+	check(CurNumPreviews < NumPreviews);		// todo: support less
+
+	FTransform OrigTransform = Cast<IPrimitiveComponentBackedTarget>(Target)->GetWorldTransform();
+	FVector WorldShift = (FVector)((WorldBounds.Width() * 1.1) * FVector3d::UnitX());
+
+	for (int32 k = CurNumPreviews; k < NumPreviews; ++k)
+	{
+		TUniquePtr<FGenerateLODOperatorFactory> Factory = MakeUnique<FGenerateLODOperatorFactory>();
+		Factory->ParentTool = this;
+
+		UMeshOpPreviewWithBackgroundCompute* NewPreview = NewObject<UMeshOpPreviewWithBackgroundCompute>(this);
+		NewPreview->Setup(this->TargetWorld, Factory.Get());
+
+		FComponentMaterialSet MaterialSet;
+		Cast<IMaterialProvider>(Target)->GetMaterialSet(MaterialSet);
+		NewPreview->ConfigureMaterials(MaterialSet.Materials,
+			ToolSetupUtil::GetDefaultWorkingMaterial(GetToolManager())
+		);
+
+		NewPreview->PreviewMesh->SetTangentsMode(EDynamicMeshComponentTangentsMode::AutoCalculated);
+		NewPreview->PreviewMesh->UpdatePreview(OriginalMesh.Get());
+
+		FTransform UseTransform = OrigTransform;
+		UseTransform.AddToTranslation( (float)(k+1)*WorldShift );
+		NewPreview->PreviewMesh->SetTransform(UseTransform);
+
+		Factory->UseTransform = UseTransform;
+
+		Previews.Add(NewPreview);
+		PreviewFactories.Add(MoveTemp(Factory)); 
+
+		NewPreview->OnMeshUpdated.AddLambda([this](UMeshOpPreviewWithBackgroundCompute* Compute)
+		{
+			OnPreviewUpdated(Compute);
+		});
+	}
+
+	for (int32 k = 0; k < Previews.Num(); ++k)
+	{
+		PreviewFactories[k]->LODSettings = SimplifyProperties->LODLevels[k];
+	}
+
+	CachedLODLevels.Reset();
+	CachedLODLevels.SetNum(Previews.Num());
+		
+	InvalidateAllPreviews();
+}
+
+
+void UGenerateLODMeshesTool::InvalidateAllPreviews()
+{
+	// just invalidate caches, next Tick() will kick off new computes
+	for (int32 k = 0; k < Previews.Num(); ++k)
+	{
+		CachedLODLevels[k] = FLODLevelGenerateSettings();
+	}
+}
+
+
+void UGenerateLODMeshesTool::OnPreviewUpdated(UMeshOpPreviewWithBackgroundCompute* PreviewCompute)
+{
+	for (int32 k = 0; k < Previews.Num(); ++k)
+	{
+		if (Previews[k] == PreviewCompute)
+		{
+			const FDynamicMesh3* Mesh = Previews[k]->PreviewMesh->GetMesh();
+			const FDynamicMeshUVOverlay* UVLayer = Mesh->Attributes()->GetUVLayer(0);
+
+			FLODLevelGenerateSettings& GenerateSettings = SimplifyProperties->LODLevels[k];
+			GenerateSettings.Result = FString::Printf(TEXT("V:%d  T:%d  U:%d"), 
+				Mesh->VertexCount(), Mesh->TriangleCount(), UVLayer->ElementCount());
+
+			return;
+		}
+	}
+}
+
+
+TUniquePtr<UE::Geometry::FDynamicMeshOperator> FGenerateLODOperatorFactory::MakeNewOperator()
+{
+	TUniquePtr<UE::Geometry::FSimplifyMeshOp> Op = MakeUnique<UE::Geometry::FSimplifyMeshOp>();
+
+	UGenerateLODMeshesToolProperties* SimplifyProperties = ParentTool->SimplifyProperties;
+
+	Op->SimplifierType = LODSettings.SimplifierType;
+	Op->TargetMode = LODSettings.TargetMode;
+	Op->TargetCount = LODSettings.TargetCount;
+	Op->TargetPercentage = LODSettings.TargetPercentage;
+	Op->bReproject = LODSettings.bReproject;
+
+	Op->TargetEdgeLength = SimplifyProperties->TargetEdgeLength;
+	Op->bDiscardAttributes = SimplifyProperties->bDiscardAttributes;
+	Op->bPreventNormalFlips = SimplifyProperties->bPreventNormalFlips;
+	Op->bPreserveSharpEdges = SimplifyProperties->bPreserveSharpEdges;
+	Op->bAllowSeamCollapse = !SimplifyProperties->bPreserveSharpEdges;
+	Op->MeshBoundaryConstraint = (EEdgeRefineFlags)SimplifyProperties->MeshBoundaryConstraint;
+	Op->GroupBoundaryConstraint = (EEdgeRefineFlags)SimplifyProperties->GroupBoundaryConstraint;
+	Op->MaterialBoundaryConstraint = (EEdgeRefineFlags)SimplifyProperties->MaterialBoundaryConstraint;
+	FTransform LocalToWorld = this->UseTransform;
+	Op->SetTransform(LocalToWorld);
+
+	Op->OriginalMeshDescription = ParentTool->OriginalMeshDescription;
+	Op->OriginalMesh = ParentTool->OriginalMesh;
+	Op->OriginalMeshSpatial = ParentTool->OriginalMeshSpatial;
+
+	IMeshReductionManagerModule& MeshReductionModule = FModuleManager::Get().LoadModuleChecked<IMeshReductionManagerModule>("MeshReductionInterface");
+	Op->MeshReduction = MeshReductionModule.GetStaticMeshReductionInterface();
+
+	return Op;
+}
+
+
+
+void UGenerateLODMeshesTool::Render(IToolsContextRenderAPI* RenderAPI)
+{
+	//FPrimitiveDrawInterface* PDI = RenderAPI->GetPrimitiveDrawInterface();
+	//FTransform Transform = Cast<IPrimitiveComponentBackedTarget>(Target)->GetWorldTransform(); //Actor->GetTransform();
+
+	//FColor LineColor(255, 0, 0);
+	//const FDynamicMesh3* TargetMesh = Preview->PreviewMesh->GetPreviewDynamicMesh();
+	//if (TargetMesh->HasAttributes())
+	//{
+	//	float PDIScale = RenderAPI->GetCameraState().GetPDIScalingFactor();
+
+	//	const FDynamicMeshUVOverlay* UVOverlay = TargetMesh->Attributes()->PrimaryUV();
+	//	for (int eid : TargetMesh->EdgeIndicesItr())
+	//	{
+	//		if (UVOverlay->IsSeamEdge(eid))
+	//		{
+	//			FVector3d A, B;
+	//			TargetMesh->GetEdgeV(eid, A, B);
+	//			PDI->DrawLine(Transform.TransformPosition((FVector)A), Transform.TransformPosition((FVector)B),
+	//				LineColor, 0, 2.0*PDIScale, 1.0f, true);
+	//		}
+	//	}
+	//}
+}
+
+
+
+void UGenerateLODMeshesTool::UpdateVisualization()
+{
+	IMaterialProvider* TargetMaterial = Cast<IMaterialProvider>(Target);
+	FComponentMaterialSet MaterialSet;
+	TargetMaterial->GetMaterialSet(MaterialSet);
+
+	for (UMeshOpPreviewWithBackgroundCompute* Preview : Previews)
+	{
+		Preview->PreviewMesh->EnableWireframe(SimplifyProperties->bShowWireframe);
+
+		Preview->ConfigureMaterials(MaterialSet.Materials,
+			ToolSetupUtil::GetDefaultWorkingMaterial(GetToolManager()));
+	}
+
+	//FComponentMaterialSet MaterialSet;
+	//if (SimplifyProperties->bShowGroupColors)
+	//{
+	//	MaterialSet.Materials = {ToolSetupUtil::GetSelectionMaterial(GetToolManager())};
+	//	Preview->PreviewMesh->SetTriangleColorFunction([this](const FDynamicMesh3* Mesh, int TriangleID)
+	//	{
+	//		return LinearColors::SelectFColor(Mesh->GetTriangleGroup(TriangleID));
+	//	},
+	//	UPreviewMesh::ERenderUpdateMode::FastUpdate);
+	//}
+	//else
+	//{
+	//	TargetMaterial->GetMaterialSet(MaterialSet);
+	//	Preview->PreviewMesh->ClearTriangleColorFunction(UPreviewMesh::ERenderUpdateMode::FastUpdate);
+	//}
+
+
+
+}
+
+bool UGenerateLODMeshesTool::CanAccept() const
+{
+	for (UMeshOpPreviewWithBackgroundCompute* Preview : Previews)
+	{
+		if (!Preview->HaveValidResult())
+		{
+			return false;
+		}
+	}
+	return Super::CanAccept();
+}
+
+void UGenerateLODMeshesTool::GenerateAssets()
+{
+	GetToolManager()->BeginUndoTransaction(LOCTEXT("GenerateLODMeshesToolTransactionName", "Simplify Mesh"));
+
+	IPrimitiveComponentBackedTarget* TargetComponent = Cast<IPrimitiveComponentBackedTarget>(Target);
+	FTransform Transform = TargetComponent->GetWorldTransform();
+
+	int32 NumLODs = Previews.Num();
+	for (int32 k = 0; k < NumLODs; ++k)
+	{
+		UMeshOpPreviewWithBackgroundCompute* Preview = Previews[k];
+
+		FDynamicMeshOpResult OpResult = Preview->Shutdown();
+		if (OpResult.Mesh->TriangleCount() == 0)
+		{
+			continue;		// failed!
+		}
+
+		FComponentMaterialSet MaterialSet;
+		Cast<IMaterialProvider>(Target)->GetMaterialSet(MaterialSet);
+
+		FString BaseName = UE::Modeling::GetComponentAssetBaseName(TargetComponent->GetOwnerComponent());
+		FString Name = FString::Printf( TEXT("%s_LOD%d"), *BaseName, (SimplifyProperties->NameIndexBase+k) );
+
+		FCreateMeshObjectParams NewMeshObjectParams;
+		NewMeshObjectParams.TargetWorld = TargetWorld;
+		NewMeshObjectParams.Transform = Transform;
+		NewMeshObjectParams.BaseName = Name;
+		NewMeshObjectParams.Materials = MaterialSet.Materials;
+		NewMeshObjectParams.SetMesh(OpResult.Mesh.Get());
+		FCreateMeshObjectResult Result = UE::Modeling::CreateMeshObject(GetToolManager(), MoveTemp(NewMeshObjectParams));
+		if (Result.IsOK() && Result.NewActor != nullptr)
+		{
+			ToolSelectionUtil::SetNewActorSelection(GetToolManager(), Result.NewActor);
+		}
+	}
+
+	GetToolManager()->EndUndoTransaction();
+}
+
+
+
+
+#undef LOCTEXT_NAMESPACE
diff --git a/Engine/Plugins/Experimental/HairModelingToolset/Source/HairModelingToolset/Private/GroomCardsEditorTool.cpp b/Engine/Plugins/Experimental/HairModelingToolset/Source/HairModelingToolset/Private/GroomCardsEditorTool.cpp
new file mode 100644
--- /dev/null
+++ b/Engine/Plugins/Experimental/HairModelingToolset/Source/HairModelingToolset/Private/GroomCardsEditorTool.cpp
@@ -0,0 +1,992 @@
+// Copyright Epic Games, Inc. All Rights Reserved.
+
+#include "GroomCardsEditorTool.h"
+#include "InteractiveToolManager.h"
+#include "ToolBuilderUtil.h"
+
+#include "BaseBehaviors/SingleClickBehavior.h"
+#include "BaseBehaviors/MouseHoverBehavior.h"
+
+#include "DynamicMesh/DynamicMesh3.h"
+#include "DynamicMesh/DynamicMeshAABBTree3.h"
+#include "DynamicMesh/MeshNormals.h"
+#include "DynamicMeshEditor.h"
+#include "Util/ColorConstants.h"
+#include "ToolSetupUtil.h"
+#include "Selection/ToolSelectionUtil.h"
+#include "DynamicMesh/MeshTransforms.h"
+#include "DynamicMeshEditor.h"
+#include "DynamicMesh/DynamicMeshChangeTracker.h"
+#include "Parameterization/DynamicMeshUVEditor.h"
+#include "Parameterization/MeshUVTransforms.h"
+#include "ModelingToolTargetUtil.h"
+#include "DynamicMeshToMeshDescription.h"
+
+#include "Changes/MeshVertexChange.h"
+
+#include "GroomQueryUtil.h"
+
+#include "Engine/Classes/Engine/StaticMesh.h"
+#include "Engine/Classes/Components/StaticMeshComponent.h"
+#include "Engine/Classes/Engine/StaticMeshActor.h"
+
+
+#define LOCTEXT_NAMESPACE "UGroomCardsEditorTool"
+
+#include "ExplicitUseGeometryMathTypes.h"		// using UE::Geometry::(math types)
+using namespace UE::Geometry;
+
+/*
+ * ToolBuilder
+ */
+
+UMeshSurfacePointTool* UGroomCardsEditorToolBuilder::CreateNewTool(const FToolBuilderState& SceneState) const
+{
+	UGroomCardsEditorTool* NewTool = NewObject<UGroomCardsEditorTool>(SceneState.ToolManager);
+	NewTool->SetWorld(SceneState.World);
+	return NewTool;
+}
+
+
+void UEditGroomCardsToolActionPropertySet::PostAction(EEditGroomCardsToolActions Action)
+{
+	if (ParentTool.IsValid())
+	{
+		ParentTool->RequestAction(Action);
+	}
+}
+
+
+
+/*
+ * Tool
+ */
+UGroomCardsEditorTool::UGroomCardsEditorTool()
+{
+	SetToolDisplayName(LOCTEXT("GroomCardsEditorToolName", "Groom to Mesh"));
+}
+
+
+
+
+void UGroomCardsEditorTool::Setup()
+{
+	UMeshSurfacePointTool::Setup();
+
+	PreviewMesh = NewObject<UPreviewMesh>(this);
+	PreviewMesh->bBuildSpatialDataStructure = true;
+	PreviewMesh->CreateInWorld(TargetWorld, FTransform::Identity);
+	PreviewMesh->SetTransform((FTransform)UE::ToolTarget::GetLocalToWorldTransform(Target));
+
+	MeshMaterial = ToolSetupUtil::GetDefaultSculptMaterial(GetToolManager());
+	UVMaterial = ToolSetupUtil::GetUVCheckerboardMaterial(50.0);
+	PreviewMesh->SetMaterial(MeshMaterial);
+
+	// Set up SelectionMechanic but remove its behaviors because we'd rather deal with clicks
+	// ourselves. Specifically, we want to be able to start new transactions on clicks, which
+	// we currently shouldn't do in an OnSelectionModifiedEvent call.
+	CardMeshSelectionMechanic = NewObject<UPolygonSelectionMechanic>(this);
+	CardMeshSelectionMechanic->bAddSelectionFilterPropertiesToParentTool = false;
+	CardMeshSelectionMechanic->Setup(this);
+	CardMeshSelectionMechanic->OnSelectionChanged.AddUObject(this, &UGroomCardsEditorTool::OnSelectionModifiedEvent);
+	CardMeshSelectionMechanic->PolyEdgesRenderer.LineThickness = 1.0;
+	CardMeshSelectionMechanic->DisableBehaviors(this);
+
+	// configure secondary render material
+	UMaterialInterface* SelectionMaterial = ToolSetupUtil::GetSelectionMaterial(FLinearColor(0.9f, 0.1f, 0.1f), GetToolManager());
+	if (SelectionMaterial != nullptr)
+	{
+		PreviewMesh->SetSecondaryRenderMaterial(SelectionMaterial);
+	}
+
+	// enable secondary triangle buffers
+	PreviewMesh->EnableSecondaryTriangleBuffers(
+		[this](const FDynamicMesh3* Mesh, int32 TriangleID)
+	{
+		return CardMeshSelectionMechanic->GetActiveSelection().IsSelectedTriangle(Mesh, Topology.Get(), TriangleID);
+	});
+
+
+	PreviewGeom = NewObject<UPreviewGeometry>(this);
+	PreviewGeom->CreateInWorld(TargetWorld, FTransform::Identity);
+	PreviewGeom->GetActor()->SetActorTransform((FTransform)UE::ToolTarget::GetLocalToWorldTransform(Target));
+
+	UE::ToolTarget::HideSourceObject(Target);
+
+	ControlPointsMechanic = NewObject<USpaceCurveDeformationMechanic>(this);
+	ControlPointsMechanic->Setup(this);
+	ControlPointsMechanic->SetWorld(TargetWorld);
+	ControlPointsMechanic->OnPointsChanged.AddLambda([this]() 
+	{
+		bCurveUpdatePending = true;
+	});
+
+
+	// register click behavior
+	USingleClickInputBehavior* ClickBehavior = NewObject<USingleClickInputBehavior>();
+	ClickBehavior->Initialize(this);
+	AddInputBehavior(ClickBehavior);
+
+	// configure behavior priorities
+	ClickBehavior->SetDefaultPriority(ControlPointsMechanic->ClickBehavior->GetPriority().MakeLower());
+	// hardcode these higher so they override the behaviors in UMeshSurfacePointTool that are not exposed
+	ControlPointsMechanic->ClickBehavior->SetDefaultPriority(FInputCapturePriority(50));
+	ControlPointsMechanic->HoverBehavior->SetDefaultPriority(FInputCapturePriority(50));
+
+
+	InfoProperties = NewObject<UGroomCardsInfoToolProperties>(this);
+	InfoProperties->RestoreProperties(this);
+	AddToolPropertySource(InfoProperties);
+
+	ControlPointsMechanic->TransformProperties->RestoreProperties(this);
+	AddToolPropertySource(ControlPointsMechanic->TransformProperties);
+
+	SelectActions = NewObject<USelectGroomCardsToolActions>();
+	SelectActions->Initialize(this);
+	AddToolPropertySource(SelectActions);
+
+	EditActions = NewObject<UEditGroomCardsToolActions>();
+	EditActions->Initialize(this);
+	AddToolPropertySource(EditActions);
+
+	bSetupValid = false;
+	bVisualizationChanged = true;
+	PendingAction = EEditGroomCardsToolActions::NoAction;
+
+	InitializeMesh();
+	bSetupValid = true;
+
+	GetToolManager()->DisplayMessage(
+		LOCTEXT("OnStartTool", "Edit Groom Cards"),
+		EToolMessageLevel::UserNotification);
+
+	OnPreviewMeshChangedHandle = PreviewMesh->GetOnMeshChanged().Add(
+		FSimpleMulticastDelegate::FDelegate::CreateUObject(this, &UGroomCardsEditorTool::OnPreviewMeshChanged));
+}
+
+
+void UGroomCardsEditorTool::Shutdown(EToolShutdownType ShutdownType)
+{
+	UMeshSurfacePointTool::Shutdown(ShutdownType);
+
+	CardMeshSelectionMechanic->Shutdown();
+
+	ControlPointsMechanic->TransformProperties->SaveProperties(this);
+	ControlPointsMechanic->Shutdown();
+
+	PreviewMesh->GetOnMeshChanged().Remove(OnPreviewMeshChangedHandle);
+
+	if (ShutdownType == EToolShutdownType::Accept)
+	{
+		GetToolManager()->BeginUndoTransaction(LOCTEXT("CardsMeshApplyEditChange", "Update Cards Mesh"));
+		PreviewMesh->ProcessMesh([&](const FDynamicMesh3& ReadMesh)
+		{
+			UE::ToolTarget::CommitDynamicMeshUpdate(Target, ReadMesh, false);
+		});
+		GetToolManager()->EndUndoTransaction();
+	}
+
+	PreviewMesh->Disconnect();
+	PreviewMesh = nullptr;
+
+	PreviewGeom->Disconnect();
+	PreviewGeom = nullptr;
+
+	UE::ToolTarget::ShowSourceObject(Target);
+}
+
+
+void UGroomCardsEditorTool::OnTick(float DeltaTime)
+{
+	if (PendingAction != EEditGroomCardsToolActions::NoAction)
+	{
+		if (PendingAction == EEditGroomCardsToolActions::Delete)
+		{
+			ApplyDeleteAction();
+		}
+		else if (PendingAction == EEditGroomCardsToolActions::SelectionClear)
+		{
+			if (ControlPointsMechanic)
+			{
+				ControlPointsMechanic->SelectionClear();
+			}
+		}
+		else if (PendingAction == EEditGroomCardsToolActions::SelectionFill)
+		{
+			if (ControlPointsMechanic)
+			{
+				ControlPointsMechanic->SelectionFill();
+			}
+		}
+		else if (PendingAction == EEditGroomCardsToolActions::SelectionAddNext)
+		{
+			if (ControlPointsMechanic)
+			{
+				ControlPointsMechanic->SelectionGrowToNext();
+			}
+		}
+		else if (PendingAction == EEditGroomCardsToolActions::SelectionAddPrevious)
+		{
+			if (ControlPointsMechanic)
+			{
+				ControlPointsMechanic->SelectionGrowToPrev();
+			}
+		}
+		else if (PendingAction == EEditGroomCardsToolActions::SelectionAddToEnd)
+		{
+			if (ControlPointsMechanic)
+			{
+				ControlPointsMechanic->SelectionGrowToEnd();
+			}
+		}
+		else if (PendingAction == EEditGroomCardsToolActions::SelectionAddToStart)
+		{
+			if (ControlPointsMechanic)
+			{
+				ControlPointsMechanic->SelectionGrowToStart();
+			}
+		}
+		PendingAction = EEditGroomCardsToolActions::NoAction;
+	}
+
+	if (bVisualizationChanged)
+	{
+		//PreviewMesh->SetMaterial(Settings->bShowUVs ? UVMaterial : MeshMaterial);
+		bVisualizationChanged = false;
+	}
+
+
+	if (bSelectionStateDirty)
+	{
+		UDynamicMeshComponent* DynamicMeshComponent = Cast<UDynamicMeshComponent>(PreviewMesh->GetRootComponent());		// ugh
+		if (DynamicMeshComponent)
+		{
+			DynamicMeshComponent->FastNotifySecondaryTrianglesChanged();
+		}
+		bSelectionStateDirty = false;
+	}
+
+	if (bCurveUpdatePending)
+	{
+		UpdateOnCurveEdit();
+		bCurveUpdatePending = false;
+	}
+
+	if (ControlPointsMechanic)
+	{
+		ControlPointsMechanic->Tick(DeltaTime);
+	}
+}
+
+
+
+
+bool UGroomCardsEditorTool::HasAccept() const
+{
+	return true;
+}
+
+bool UGroomCardsEditorTool::CanAccept() const
+{
+	return true;
+}
+
+
+
+
+
+
+
+bool UGroomCardsEditorTool::HitTest(const FRay& WorldRay, FHitResult& OutHit)
+{
+	return CardMeshSelectionMechanic->TopologyHitTest(WorldRay, OutHit);
+}
+
+
+
+FInputRayHit UGroomCardsEditorTool::IsHitByClick(const FInputDeviceRay& ClickPos)
+{
+	FHitResult OutHit;
+	if (HitTest(ClickPos.WorldRay, OutHit))
+	{
+		return FInputRayHit(OutHit.Distance);
+	}
+
+	// background capture, if nothing else is hit
+	return FInputRayHit(TNumericLimits<float>::Max());
+}
+
+
+
+
+
+
+
+FInputRayHit UGroomCardsEditorTool::CanBeginClickDragSequence(const FInputDeviceRay& PressPos)
+{
+	return FInputRayHit();
+}
+
+void UGroomCardsEditorTool::OnBeginDrag(const FRay& WorldRay)
+{
+	check(false);
+}
+
+void UGroomCardsEditorTool::OnUpdateDrag(const FRay& Ray)
+{
+	check(false);
+}
+
+void UGroomCardsEditorTool::OnEndDrag(const FRay& Ray)
+{
+	check(false);
+}
+
+
+bool UGroomCardsEditorTool::OnUpdateHover(const FInputDeviceRay& DevicePos)
+{
+	//if (ActiveVertexChange == nullptr && MultiTransformer->InGizmoEdit() == false)
+	//{
+	CardMeshSelectionMechanic->UpdateHighlight(DevicePos.WorldRay);
+	//}
+	return true;
+}
+
+void UGroomCardsEditorTool::OnEndHover()
+{
+	CardMeshSelectionMechanic->ClearHighlight();
+}
+
+
+
+
+
+
+class FEditableGroomCardSet
+{
+public:
+
+	TSharedPtr<FDynamicMesh3> FullCardMesh;
+
+	TSharedPtr<UE::GroomQueries::FMeshCardStripSet> CardStrips;
+
+
+	void Initialize(const FMeshDescription* MeshDescripIn)
+	{
+		FullCardMesh = MakeShared<FDynamicMesh3>();
+		FullCardMesh->EnableTriangleGroups(0);
+		FullCardMesh->EnableAttributes();
+
+		CardStrips = MakeShared<UE::GroomQueries::FMeshCardStripSet>();
+
+		UE::GroomQueries::ExtractAllHairCards(MeshDescripIn, *FullCardMesh, *CardStrips);
+		UE::GroomQueries::ExtractCardQuads(*FullCardMesh, *CardStrips);
+		UE::GroomQueries::ExtractCardCurves(*FullCardMesh, *CardStrips);
+	}
+};
+
+
+
+struct FCardCurvePointInfo
+{
+	FIndex2i Vertices;
+	FVector3d LocalPoints[2];
+};
+
+class FGroomCardEdit
+{
+public:
+	int32 CardGroupID;
+
+	UE::GroomQueries::FMeshCardStrip InitialCard;
+
+	UE::GroomQueries::FMeshCardStrip EditedCard;
+
+	TArray<FFrame3d> InitialFrames;
+	TArray<FFrame3d> EditedFrames;
+
+	TArray<FCardCurvePointInfo> CurveMeshVertices;
+
+	void InitializeEncoding(const FDynamicMesh3* Mesh)
+	{
+		int32 NumCurveVerts = InitialCard.CardCurve.Num();
+		CurveMeshVertices.SetNum(NumCurveVerts);
+		for ( int32 k = 0; k < NumCurveVerts; ++k )
+		{
+			FCardCurvePointInfo VertInfo;
+			VertInfo.Vertices = InitialCard.GetCurvePointVertices(k);
+			VertInfo.LocalPoints[0] = InitialFrames[k].ToFramePoint(Mesh->GetVertex(VertInfo.Vertices.A));
+			VertInfo.LocalPoints[1] = InitialFrames[k].ToFramePoint(Mesh->GetVertex(VertInfo.Vertices.B));
+			CurveMeshVertices[k] = VertInfo;
+		}
+	}
+
+	void UpdateFromEditedFrames(FDynamicMesh3& Mesh, FMeshVertexChangeBuilder* Tracker)
+	{
+		int32 NumCurveVerts = InitialCard.CardCurve.Num();
+		for (int32 k = 0; k < NumCurveVerts; ++k)
+		{
+			for (int32 j = 0; j < 2; ++j)
+			{
+				int32 VertIndex = CurveMeshVertices[k].Vertices[j];
+				FVector3d LocalPoint = CurveMeshVertices[k].LocalPoints[j];
+				FVector3d MeshPoint = EditedFrames[k].FromFramePoint(LocalPoint);
+				if (Tracker)
+				{
+					Tracker->SaveVertexInitial(&Mesh, VertIndex);
+				}
+				Mesh.SetVertex(VertIndex, MeshPoint);
+			}
+
+			EditedCard.CardCurve[k] = EditedFrames[k].Origin;
+			EditedCard.CardCurveFrames[k] = EditedFrames[k];
+		}
+
+		// todo: normals?
+	}
+};
+
+
+
+
+class FGroomCardEditChange : public FToolCommandChange
+{
+public:
+	FGroomCardEdit Edit;
+
+	virtual void Apply(UObject* Object) override
+	{
+		Cast<UGroomCardsEditorTool>(Object)->ApplyGroomCardEdit(Edit, false);
+	}
+
+	virtual void Revert(UObject* Object) override
+	{
+		Cast<UGroomCardsEditorTool>(Object)->ApplyGroomCardEdit(Edit, true);
+	}
+};
+
+
+
+
+
+
+void UGroomCardsEditorTool::OnClicked(const FInputDeviceRay& ClickPos)
+{
+	// update selection
+	GetToolManager()->BeginUndoTransaction(LOCTEXT("CardsEditSelectionChange", "Selection"));
+
+	CardMeshSelectionMechanic->BeginChange();
+	FVector3d LocalHitPosition, LocalHitNormal;
+	bool bSelectionModified = CardMeshSelectionMechanic->UpdateSelection(ClickPos.WorldRay, LocalHitPosition, LocalHitNormal);
+	CardMeshSelectionMechanic->EndChangeAndEmitIfModified();
+
+	if (bSelectionModified)
+	{
+		// close any pending card edit
+		CompleteActiveCardEditAndEmitChanges();
+
+		// if we have a new single selection, start a new edit
+		const FGroupTopologySelection& CurSelection = CardMeshSelectionMechanic->GetActiveSelection();
+		if (CurSelection.SelectedGroupIDs.Num() == 1)
+		{
+			for (const int32 SelectedCardGroupID : CurSelection.SelectedGroupIDs)
+			{
+				GetToolManager()->EmitObjectChange(this, MakeUnique<FBeginGroomCardsEditChange>(SelectedCardGroupID),
+					LOCTEXT("BeginCardEditChange", "BeginCardEdit"));
+
+				BeginCardEdit(SelectedCardGroupID);
+				break;
+			}
+		}
+
+	}
+
+	GetToolManager()->EndUndoTransaction();
+}
+
+
+
+
+void UGroomCardsEditorTool::CompleteActiveCardEditAndEmitChanges()
+{
+	// if we have an active edit, complete it
+	if (ActiveCardEdit)
+	{
+		// save current selection from control points mechanic first, so that when we undo, 
+		// control points will be initalized before we restore selection
+		ControlPointsMechanic->ClearSelection();
+
+		TUniquePtr<FEndGroomCardsEditChange> EndCardEditChange = MakeUnique<FEndGroomCardsEditChange>();
+		EndCardEditChange->CardEdit = MakePimpl<FGroomCardEdit>();
+		*(EndCardEditChange->CardEdit) = *ActiveCardEdit;
+		GetToolManager()->EmitObjectChange(this, MoveTemp(EndCardEditChange), LOCTEXT("EndCardEditChange", "EndCardEdit"));
+
+		// save change to mesh and card set if they were modified
+		if (bActiveCardEditUpdated)
+		{
+			//EndMoveChange();
+			PreviewMesh->ForceRebuildSpatial();
+
+			TUniquePtr<FGroomCardEditChange> CardSetUpdateChange = MakeUnique<FGroomCardEditChange>();
+			CardSetUpdateChange->Edit = *ActiveCardEdit;
+			CardSetUpdateChange->Apply(this);
+			GetToolManager()->EmitObjectChange(this, MoveTemp(CardSetUpdateChange), LOCTEXT("CardSetUpdateChange", "UpdateCardSet"));
+		}
+
+		EndCardEdit();
+	}
+}
+
+
+
+
+void UGroomCardsEditorTool::BeginCardEdit(int32 CardGroupID)
+{
+	const UE::GroomQueries::FMeshCardStrip& CardStrip = EditableCardSet->CardStrips->FindStripForGroup(CardGroupID);
+
+	UE::Geometry::FTransform3d Transform(PreviewMesh->GetTransform());
+	check(FMathd::Abs(Transform.TransformVector(UE::Geometry::Normalized(FVector3d(1,1,1))).SquaredLength() - 1.0) < FMathd::ZeroTolerance);
+	FFrame3d TransformFrame(Transform.GetTranslation(), Transform.GetRotation());
+
+	TArray<FFrame3d> LocalFrames;
+	for (FFrame3d LocalFrame : CardStrip.CardCurveFrames)
+	{
+		LocalFrames.Add(LocalFrame);
+	}
+
+	// initialize current edit
+	ActiveCardEdit = MakePimpl<FGroomCardEdit>();
+	ActiveCardEdit->CardGroupID = CardGroupID;
+	ActiveCardEdit->InitialCard = CardStrip;
+	ActiveCardEdit->EditedCard = ActiveCardEdit->InitialCard;
+	ActiveCardEdit->InitialFrames = LocalFrames;
+	ActiveCardEdit->EditedFrames = ActiveCardEdit->InitialFrames;
+	ActiveCardEdit->InitializeEncoding(PreviewMesh->GetMesh());
+
+	// initialize new curve adapter
+	CurveSourceAdapter = MakeShared<FSpaceCurveSource>();
+	CurveSourceAdapter->GetPointCount = [this]() { check(ActiveCardEdit.IsValid());  return ActiveCardEdit->EditedFrames.Num(); };
+	CurveSourceAdapter->GetPoint = [this, TransformFrame](int32 Index) { return TransformFrame.FromFrame(ActiveCardEdit->EditedFrames[Index]); };
+	CurveSourceAdapter->IsLoop = []() { return false; };
+	ControlPointsMechanic->SetCurveSource(CurveSourceAdapter);
+
+	ActiveCardGroupID = CardGroupID;
+	bActiveCardEditUpdated = false;
+}
+
+
+
+void UGroomCardsEditorTool::RestoreCardEdit(const FGroomCardEdit* RestoreEdit)
+{
+	// this should probably be cached once...
+	UE::Geometry::FTransform3d Transform(PreviewMesh->GetTransform());
+	check(FMathd::Abs(Transform.TransformVector(UE::Geometry::Normalized(FVector3d(1, 1, 1))).SquaredLength() - 1.0) < FMathd::ZeroTolerance);
+	FFrame3d TransformFrame(Transform.GetTranslation(), Transform.GetRotation());
+
+	ActiveCardEdit = MakePimpl<FGroomCardEdit>();
+	*ActiveCardEdit = *RestoreEdit;
+
+	// initialize new curve adapter
+	CurveSourceAdapter = MakeShared<FSpaceCurveSource>();
+	CurveSourceAdapter->GetPointCount = [this]() { check(ActiveCardEdit.IsValid());  return ActiveCardEdit->EditedFrames.Num(); };
+	CurveSourceAdapter->GetPoint = [this, TransformFrame](int32 Index) { return TransformFrame.FromFrame(ActiveCardEdit->EditedFrames[Index]); };
+	CurveSourceAdapter->IsLoop = []() { return false; };
+	ControlPointsMechanic->SetCurveSource(CurveSourceAdapter);
+
+	// update mesh with this edit result
+	//UpdateOnCurveEdit();
+	bCurveUpdatePending = true;
+
+	ActiveCardGroupID = ActiveCardEdit->CardGroupID;
+	bActiveCardEditUpdated = false;
+}
+
+
+
+
+void UGroomCardsEditorTool::EndCardEdit()
+{
+	ActiveCardEdit = nullptr;
+	ControlPointsMechanic->ClearCurveSource();
+	CurveSourceAdapter = nullptr;
+	ActiveCardGroupID = -1;
+	bActiveCardEditUpdated = false;
+}
+
+
+
+
+
+void UGroomCardsEditorTool::UpdateOnCurveEdit()
+{
+	check(ActiveCardEdit && ActiveCardEdit.IsValid());
+
+	if (bActiveCardEditUpdated == false)
+	{
+		bActiveCardEditUpdated = true;
+		//BeginMoveChange();
+	}
+
+	UE::Geometry::FTransform3d Transform(PreviewMesh->GetTransform());
+	check( FMathd::Abs(Transform.TransformVector(UE::Geometry::Normalized(FVector3d(1,1,1))).SquaredLength() - 1.0) < FMathd::ZeroTolerance );
+	FFrame3d TransformFrame(Transform.GetTranslation(), Transform.GetRotation());
+
+	TArray<FFrame3d> NewPositions;
+	ControlPointsMechanic->GetCurrentCurvePoints(NewPositions);
+
+	for (int32 k = 0; k < ActiveCardEdit->EditedFrames.Num(); ++k)
+	{
+		FFrame3d LocalFrame = TransformFrame.ToFrame(NewPositions[k]);
+		ActiveCardEdit->EditedFrames[k] = LocalFrame;
+	}
+
+	PreviewMesh->DeferredEditMesh([&](FDynamicMesh3& Mesh)
+	{
+		//ActiveCardEdit->UpdateFromEditedFrames(Mesh, ActiveVertexChange.Get());
+		ActiveCardEdit->UpdateFromEditedFrames(Mesh, nullptr);
+	}, false);
+
+	PreviewMesh->NotifyDeferredEditCompleted(UPreviewMesh::ERenderUpdateMode::FastUpdate, 
+		EMeshRenderAttributeFlags::Positions, false);
+}
+
+
+
+void UGroomCardsEditorTool::ApplyGroomCardEdit(const FGroomCardEdit& Edit, bool bIsRevert)
+{
+	UE::GroomQueries::FMeshCardStrip& CardStrip = EditableCardSet->CardStrips->FindStripForGroup(Edit.CardGroupID);
+
+	const UE::GroomQueries::FMeshCardStrip& SetToStrip = (bIsRevert) ? Edit.InitialCard : Edit.EditedCard;
+
+	CardStrip.CardCurve = SetToStrip.CardCurve;
+	CardStrip.CardCurveFrames = SetToStrip.CardCurveFrames;
+}
+
+
+
+
+void UGroomCardsEditorTool::OnSelectionModifiedEvent()
+{
+	bSelectionStateDirty = true;
+}
+
+
+
+
+
+void UGroomCardsEditorTool::BeginMoveChange()
+{
+	check(ActiveVertexChange == nullptr);
+	ActiveVertexChange = new FMeshVertexChangeBuilder(EMeshVertexChangeComponents::VertexPositions | EMeshVertexChangeComponents::OverlayNormals);
+}
+
+void UGroomCardsEditorTool::EndMoveChange()
+{
+	check(ActiveVertexChange);
+	if (ActiveVertexChange->SavedVertices.Num() > 0)
+	{
+
+		TUniquePtr<TWrappedToolCommandChange<FMeshVertexChange>> NewChange = MakeUnique<TWrappedToolCommandChange<FMeshVertexChange>>();
+		NewChange->WrappedChange = MoveTemp(ActiveVertexChange->Change);
+		//NewChange->BeforeModify = [this](bool bRevert)
+		//{
+		//	this->WaitForPendingUndoRedo();
+		//};
+
+		GetToolManager()->EmitObjectChange(PreviewMesh, MoveTemp(NewChange), LOCTEXT("CardVertexEditChange", "Edit Card"));
+
+		PreviewMesh->ForceRebuildSpatial();
+	}
+
+	delete ActiveVertexChange;
+	ActiveVertexChange = nullptr;
+}
+
+
+
+
+
+
+
+void UGroomCardsEditorTool::Render(IToolsContextRenderAPI* RenderAPI)
+{
+	FPrimitiveDrawInterface* PDI = RenderAPI->GetPrimitiveDrawInterface();
+	FTransform Transform = PreviewMesh->GetTransform();
+	const FDynamicMesh3* Mesh = PreviewMesh->GetMesh();
+
+	auto DrawCardQuadStripFunc = [&](const UE::GroomQueries::FMeshCardStrip& CardStrip)
+	{
+		for (FIndex4i Quad : CardStrip.QuadLoops)
+		{
+			for (int32 j = 0; j < 4; ++j)
+			{
+				FVector Pos0 = (FVector)Mesh->GetVertex(Quad[j]);
+				FVector Pos1 = (FVector)Mesh->GetVertex(Quad[(j + 1) % 4]);
+				const FLinearColor& UseColor = (j == 0 || j == 2) ? FLinearColor::Green : FLinearColor::Blue;
+				PDI->DrawLine(Transform.TransformPosition(Pos0), Transform.TransformPosition(Pos1),
+					UseColor, SDPG_Foreground, 1.0f, 0.0f, true);
+			}
+		}
+	};
+
+	if (CardMeshSelectionMechanic->HasSelection())
+	{
+		for (int32 SelectedGroup : CardMeshSelectionMechanic->GetActiveSelection().SelectedGroupIDs)
+		{
+			DrawCardQuadStripFunc(EditableCardSet->CardStrips->FindStripForGroup(SelectedGroup));
+		}
+	}
+
+
+	if (ControlPointsMechanic != nullptr)
+	{
+		ControlPointsMechanic->Render(RenderAPI);
+	}
+
+}
+
+
+
+void UGroomCardsEditorTool::UpdateLineSet()
+{
+}
+
+
+
+
+void UGroomCardsEditorTool::InitializeMesh()
+{
+	EditableCardSet = MakePimpl<FEditableGroomCardSet>();
+	EditableCardSet->Initialize(UE::ToolTarget::GetMeshDescription(Target));
+	const FDynamicMesh3* CardsMesh = EditableCardSet->FullCardMesh.Get();
+
+	PreviewMesh->UpdatePreview(CardsMesh);
+
+	RecomputeTopology();
+
+	UpdateLineSet();
+
+	if (CardsMesh->TriangleCount() == 0)
+	{
+		GetToolManager()->DisplayMessage(
+			LOCTEXT("OnNoCardsMesage", "Selected LOD does not contain any Card geometry"),
+			EToolMessageLevel::UserWarning);
+	}
+	else
+	{
+		GetToolManager()->DisplayMessage(FText(), EToolMessageLevel::UserWarning);
+	}
+
+	InfoProperties->NumCards = EditableCardSet->CardStrips->CardStrips.Num();
+	InfoProperties->NumVertices = EditableCardSet->FullCardMesh->VertexCount();
+	InfoProperties->NumTriangles = EditableCardSet->FullCardMesh->TriangleCount();
+}
+
+
+
+
+
+
+
+
+
+
+
+void UGroomCardsEditorTool::RequestAction(EEditGroomCardsToolActions ActionType)
+{
+	if (PendingAction != EEditGroomCardsToolActions::NoAction)
+	{
+		return;
+	}
+
+	PendingAction = ActionType;
+}
+
+
+
+
+
+
+
+void UGroomCardsEditorTool::ApplyDeleteAction()
+{
+	if (CardMeshSelectionMechanic->HasSelection() == false)
+	{
+		GetToolManager()->DisplayMessage(
+			LOCTEXT("OnDeleteFailedMessage", "Cannot Delete Current Selection"),
+			EToolMessageLevel::UserWarning);
+		return;
+	}
+
+	const FGroupTopologySelection& Selection = CardMeshSelectionMechanic->GetActiveSelection();
+
+	TUniquePtr<FMeshChange> MeshChange = 
+		PreviewMesh->TrackedEditMesh([&](FDynamicMesh3& EditMesh, FDynamicMeshChangeTracker& ChangeTracker)
+	{
+		FDynamicMeshEditor Editor(&EditMesh);
+		for (int32 SelectedGroup : Selection.SelectedGroupIDs)
+		{
+			const UE::GroomQueries::FMeshCardStrip& CardStrip = EditableCardSet->CardStrips->FindStripForGroup(SelectedGroup);
+
+			ChangeTracker.SaveTriangles(CardStrip.Triangles, true);
+			Editor.RemoveTriangles(CardStrip.Triangles, true);
+		}
+	});
+	
+	// emit undo
+	FGroupTopologySelection NewSelection;
+	CompleteMeshEditChange(LOCTEXT("StripDeleteChange", "Delete"), MoveTemp(MeshChange), NewSelection);
+}
+
+
+
+
+void UGroomCardsEditorTool::RecomputeTopology()
+{
+	const FDynamicMesh3* Mesh = PreviewMesh->GetMesh();
+
+	Topology = MakeUnique<FGroupTopology>(Mesh, false);
+	// we already know the topology so this is very inefficient!
+	Topology->RebuildTopology();
+
+	// update selection mechanic
+	CardMeshSelectionMechanic->Initialize(
+		PreviewMesh->GetMesh(), PreviewMesh->GetTransform(), TargetWorld,
+		Topology.Get(),
+		[this]() { return PreviewMesh->GetSpatial(); }
+	);
+}
+
+
+
+void UGroomCardsEditorTool::OnPreviewMeshChanged()
+{
+	CardMeshSelectionMechanic->NotifyMeshChanged(false);
+}
+
+
+void UGroomCardsEditorTool::AfterTopologyEdit()
+{
+	bWasTopologyEdited = true;
+	CardMeshSelectionMechanic->NotifyMeshChanged(true);
+
+	RecomputeTopology();
+}
+
+
+
+
+void UGroomCardsEditorTool::CompleteMeshEditChange(
+	const FText& TransactionLabel,
+	TUniquePtr<FToolCommandChange> EditChange,
+	const FGroupTopologySelection& OutputSelection)
+{
+	// open top-level transaction
+	GetToolManager()->BeginUndoTransaction(TransactionLabel);
+
+	// is this the right place? seems like maybe it needs to happen before we actually do mesh edit...
+	CompleteActiveCardEditAndEmitChanges();
+
+	// clear current selection
+	CardMeshSelectionMechanic->BeginChange();
+	CardMeshSelectionMechanic->ClearSelection();
+	GetToolManager()->EmitObjectChange(CardMeshSelectionMechanic, CardMeshSelectionMechanic->EndChange(), LOCTEXT("GroomCardsEditChangeClearSelection", "ClearSelection"));
+
+	// emit the pre-edit change
+	GetToolManager()->EmitObjectChange(this, MakeUnique<FEditGroomCardsTopologyPreEditChange>(), LOCTEXT("GroomCardsEditChangePreEdit", "PreEdit"));
+
+	// emit the mesh change
+	GetToolManager()->EmitObjectChange(PreviewMesh, MoveTemp(EditChange), TransactionLabel);
+
+	// emit the post-edit change
+	GetToolManager()->EmitObjectChange(this, MakeUnique<FEditGroomCardsTopologyPostEditChange>(), TransactionLabel);
+	// call this (PostEditChange will do this)
+	AfterTopologyEdit();
+	// increment topology-change counter
+	ModifiedTopologyCounter++;
+
+	// set output selection
+	if (OutputSelection.IsEmpty() == false)
+	{
+		CardMeshSelectionMechanic->BeginChange();
+		CardMeshSelectionMechanic->SetSelection(OutputSelection);
+		GetToolManager()->EmitObjectChange(CardMeshSelectionMechanic, CardMeshSelectionMechanic->EndChange(), LOCTEXT("PolyMeshExtrudeChangeSetSelection", "SetSelection"));
+	}
+
+	// complete the transaction
+	GetToolManager()->EndUndoTransaction();
+
+	// clean up preview mesh, hiding of things, etc
+	//DynamicMeshComponent->SetSecondaryBuffersVisibility(true);
+
+	CurrentOperationTimestamp++;
+}
+
+
+
+
+
+
+
+void FEditGroomCardsTopologyPreEditChange::Apply(UObject* Object)
+{
+}
+void FEditGroomCardsTopologyPreEditChange::Revert(UObject* Object)
+{
+	Cast<UGroomCardsEditorTool>(Object)->AfterTopologyEdit();
+	Cast<UGroomCardsEditorTool>(Object)->ModifiedTopologyCounter--;
+}
+FString FEditGroomCardsTopologyPreEditChange::ToString() const
+{
+	return TEXT("FEditGroomCardsTopologyPreEditChange");
+}
+
+
+void FEditGroomCardsTopologyPostEditChange::Apply(UObject* Object)
+{
+	Cast<UGroomCardsEditorTool>(Object)->AfterTopologyEdit();
+	Cast<UGroomCardsEditorTool>(Object)->ModifiedTopologyCounter++;
+}
+void FEditGroomCardsTopologyPostEditChange::Revert(UObject* Object)
+{
+}
+FString FEditGroomCardsTopologyPostEditChange::ToString() const
+{
+	return TEXT("FEditGroomCardsTopologyPostEditChange");
+}
+
+
+
+
+void FBeginGroomCardsEditChange::Apply(UObject* Object)
+{
+	Cast<UGroomCardsEditorTool>(Object)->BeginCardEdit(CardGroupID);
+}
+void FBeginGroomCardsEditChange::Revert(UObject* Object)
+{
+	Cast<UGroomCardsEditorTool>(Object)->EndCardEdit();
+}
+FString FBeginGroomCardsEditChange::ToString() const
+{
+	return TEXT("FBeginGroomCardsEditChange");
+}
+
+
+
+void FEndGroomCardsEditChange::Apply(UObject* Object)
+{
+	Cast<UGroomCardsEditorTool>(Object)->EndCardEdit();
+}
+void FEndGroomCardsEditChange::Revert(UObject* Object)
+{
+	Cast<UGroomCardsEditorTool>(Object)->RestoreCardEdit(CardEdit.Get());
+}
+FString FEndGroomCardsEditChange::ToString() const
+{
+	return TEXT("FBeginGroomCardsEditChange");
+}
+
+
+
+
+#undef LOCTEXT_NAMESPACE
diff --git a/Engine/Plugins/Experimental/HairModelingToolset/Source/HairModelingToolset/Private/GroomQueryUtil.cpp b/Engine/Plugins/Experimental/HairModelingToolset/Source/HairModelingToolset/Private/GroomQueryUtil.cpp
new file mode 100644
--- /dev/null
+++ b/Engine/Plugins/Experimental/HairModelingToolset/Source/HairModelingToolset/Private/GroomQueryUtil.cpp
@@ -0,0 +1,575 @@
+// Copyright Epic Games, Inc. All Rights Reserved.
+
+#include "GroomQueryUtil.h"
+
+#include "DynamicMesh/DynamicMeshAttributeSet.h"
+#include "MeshDescriptionToDynamicMesh.h"
+#include "Selections/MeshConnectedComponents.h"
+#include "FaceGroupUtil.h"
+#include "HairCardsBuilder.h"
+
+#include "ExplicitUseGeometryMathTypes.h"		// using UE::Geometry::(math types)
+using namespace UE::Geometry;
+
+const UE::GroomQueries::FMeshCardStrip& UE::GroomQueries::FMeshCardStripSet::FindStripForGroup(int32 GroupIdx) const
+{
+	int32 CardIdx = TriGroupToCardIndex[GroupIdx];
+	return CardStrips[CardIdx];
+}
+
+UE::GroomQueries::FMeshCardStrip& UE::GroomQueries::FMeshCardStripSet::FindStripForGroup(int32 GroupIdx)
+{
+	int32 CardIdx = TriGroupToCardIndex[GroupIdx];
+	return CardStrips[CardIdx];
+}
+
+void UE::GroomQueries::ExtractAllHairCards(AGroomActor* GroomActor,
+	int32 LODIndex,
+	FDynamicMesh3& MeshOut,
+	FMeshCardStripSet* CardsOut)
+{
+	check(GroomActor->GetGroomComponent());
+	check(GroomActor->GetGroomComponent()->GroomAsset);
+	UGroomAsset* Asset = GroomActor->GetGroomComponent()->GroomAsset;
+
+	MeshOut.EnableTriangleGroups(0);
+	MeshOut.EnableAttributes();
+	FDynamicMeshUVOverlay* UVs = MeshOut.Attributes()->PrimaryUV();
+	FDynamicMeshNormalOverlay* Normals = MeshOut.Attributes()->PrimaryNormals();
+
+	int32 NumHairGroups = Asset->HairGroupsInfo.Num();
+	for (int32 GroupIdx = 0; GroupIdx < NumHairGroups; ++GroupIdx)
+	{
+		const FHairGroupInfo& GroupInfo = Asset->HairGroupsInfo[GroupIdx];
+		const FHairGroupData& GroupData = Asset->HairGroupsData[GroupIdx];
+
+		int32 NumCardsLODs = GroupData.Cards.LODs.Num();
+		if (LODIndex >= NumCardsLODs)
+		{
+			continue;
+		}
+
+		UStaticMesh* StaticMesh = nullptr;
+		for (const FHairGroupsCardsSourceDescription& Desc : Asset->HairGroupsCards)
+		{
+			if (Desc.GroupIndex == GroupIdx && Desc.LODIndex == LODIndex)
+			{
+				if (Desc.SourceType == EHairCardsSourceType::Imported)
+				{
+					StaticMesh = Desc.ImportedMesh;
+				}
+				else if (Desc.SourceType == EHairCardsSourceType::Procedural)
+				{
+					StaticMesh = Desc.ProceduralMesh;
+				}
+			}
+		}
+		
+		if (!StaticMesh)
+		{
+			continue;
+		}
+
+		FHairCardsDatas CardData;
+		if (!FHairCardsBuilder::ExtractCardsData(StaticMesh, CardData))
+		{
+			continue;
+		}
+		const FHairCardsGeometry& CardGeo = CardData.Cards;
+
+		int32 NumVerts = CardGeo.GetNumVertices();
+		TArray<int32> MapV, MapN, MapU;
+		MapV.SetNum(NumVerts);
+		MapU.SetNum(NumVerts);
+		MapN.SetNum(NumVerts);
+		for (int32 k = 0; k < NumVerts; ++k)
+		{
+			int32 NewVID = MeshOut.AppendVertex((FVector3d)CardGeo.Positions[k]);
+			MapV[k] = NewVID;
+
+			int32 NewNormID = Normals->AppendElement((FVector3f)CardGeo.Normals[k]);
+			MapN[k] = NewNormID;
+
+			int32 NewUVID = UVs->AppendElement( FVector2f(CardGeo.UVs[k].X, CardGeo.UVs[k].Y) );
+			MapU[k] = NewUVID;
+		}
+
+		int32 NumCards = CardGeo.IndexOffsets.Num();
+
+		TArray<int32> CardTriGroups;
+		CardTriGroups.SetNum(NumCards);
+		for (int32 k = 0; k < NumCards; ++k)
+		{
+			CardTriGroups[k] = MeshOut.AllocateTriangleGroup();
+		}
+
+		int32 NumTris = CardGeo.GetNumTriangles();
+
+		TArray<int32> CardTriVertIndexToGroup;
+		CardTriVertIndexToGroup.SetNum(3 * NumTris);
+		for (int32 k = 0; k < NumCards; ++k)
+		{
+			int32 Offset = CardGeo.IndexOffsets[k];
+			int32 Count = CardGeo.IndexCounts[k];
+			for (int32 j = 0; j < Count; ++j)
+			{
+				CardTriVertIndexToGroup[Offset + j] = CardTriGroups[k];
+			}
+		}
+
+		if (CardsOut != nullptr)
+		{
+			CardsOut->CardStrips.Reserve(CardsOut->CardStrips.Num() + NumCards);
+			for (int32 k = 0; k < NumCards; ++k)
+			{
+				int32 NewCardStripIndex = CardsOut->CardStrips.Num();
+				FMeshCardStrip& NewCardStrip = CardsOut->CardStrips.Emplace_GetRef();
+				NewCardStrip.GroupID = CardTriGroups[k];
+				NewCardStrip.SourceInfo.GroupIndex = GroupIdx;
+				NewCardStrip.SourceInfo.LODIndex = LODIndex;
+				NewCardStrip.SourceInfo.CardIndex = k;
+				NewCardStrip.SourceInfo.IndexOffset = CardGeo.IndexOffsets[k];
+				NewCardStrip.SourceInfo.IndexCount = CardGeo.IndexCounts[k];
+
+				CardsOut->TriGroupToCardIndex.Add(NewCardStrip.GroupID, NewCardStripIndex);
+			}
+		}
+
+		for (int32 k = 0; k < NumTris; ++k)
+		{
+			int32 GroupID = CardTriVertIndexToGroup[3 * k];
+			FIndex3i CardTri(CardGeo.Indices[3 * k], CardGeo.Indices[3 * k + 1], CardGeo.Indices[3 * k + 2]);
+
+			int32 NewTID = MeshOut.AppendTriangle(MapV[CardTri.A], MapV[CardTri.B], MapV[CardTri.C], GroupID);
+			Normals->SetTriangle(NewTID, FIndex3i(MapN[CardTri.A], MapN[CardTri.B], MapN[CardTri.C]));
+			UVs->SetTriangle(NewTID, FIndex3i(MapU[CardTri.A], MapU[CardTri.B], MapU[CardTri.C]));
+
+			if (CardsOut)
+			{
+				int32 CardStripIndex = CardsOut->TriGroupToCardIndex[GroupID];
+				CardsOut->CardStrips[CardStripIndex].Triangles.Add(NewTID);
+			}
+		}
+	}
+
+}
+
+
+
+
+void UE::GroomQueries::ExtractAllHairCards(
+	const FMeshDescription* SourceMeshIn,
+	FDynamicMesh3& MeshOut,
+	FMeshCardStripSet& CardsOut)
+{
+	FMeshDescriptionToDynamicMesh Converter;
+	Converter.Convert(SourceMeshIn, MeshOut);
+
+	MeshOut.EnableTriangleGroups();
+
+	FMeshConnectedComponents Components(&MeshOut);
+	Components.FindConnectedTriangles();
+	int32 NumCards = Components.Num();
+
+
+	CardsOut.CardStrips.Reserve(NumCards);
+	for (int32 k = 0; k < NumCards; ++k)
+	{
+		int32 CardTriGroup = MeshOut.AllocateTriangleGroup();
+
+		FaceGroupUtil::SetGroupID(MeshOut, Components.GetComponent(k).Indices, CardTriGroup);
+
+		int32 NewCardStripIndex = CardsOut.CardStrips.Num();
+		FMeshCardStrip& NewCardStrip = CardsOut.CardStrips.Emplace_GetRef();
+
+		NewCardStrip.GroupID = CardTriGroup;
+		NewCardStrip.SourceInfo.GroupIndex = 0;
+		NewCardStrip.SourceInfo.LODIndex = 0;
+		NewCardStrip.SourceInfo.CardIndex = k;
+		NewCardStrip.SourceInfo.IndexOffset = -1;
+		NewCardStrip.SourceInfo.IndexCount = -1;
+
+		CardsOut.TriGroupToCardIndex.Add(NewCardStrip.GroupID, NewCardStripIndex);
+
+		NewCardStrip.Triangles = Components.GetComponent(k).Indices;
+
+		// if we have an odd number of tris we are going to discard the last triangle and hope for the best
+		if (ensure(NewCardStrip.Triangles.Num() % 2 == 0) == false)
+		{
+			NewCardStrip.Triangles.SetNum(NewCardStrip.Triangles.Num() - 1);
+		}
+	}
+}
+
+
+
+static bool QuadsAreConnected(const FDynamicMesh3& Mesh, FIndex2i Quad1, FIndex2i Quad2)
+{
+	for (int32 j = 0; j < 2; ++j)
+	{
+		int32 Tri = Quad1[j];
+		FIndex3i TriNbrsTris = Mesh.GetTriNeighbourTris(Quad1[j]);
+		if (IndexUtil::FindTriIndex(Quad2.A, TriNbrsTris) != IndexConstants::InvalidID ||
+			IndexUtil::FindTriIndex(Quad2.B, TriNbrsTris) != IndexConstants::InvalidID)
+		{
+			return true;
+		}
+	}
+	return false;
+}
+
+static FIndex3i FindSharedQuadEdge(const FDynamicMesh3& Mesh, FIndex4i QuadA, FIndex4i QuadB)
+{
+	for (int32 j = 0; j < 4; ++j)
+	{
+		int32 A = QuadA[j], B = QuadA[(j + 1) % 4];
+		for (int32 k = 0; k < 4; ++k)
+		{
+			int32 C = QuadB[k], D = QuadB[(k + 1) % 4];
+			if (A == C && B == D)
+			{
+				return FIndex3i(j, k, 1);
+			}
+			else if (A == D && B == C)
+			{
+				return FIndex3i(j, k, -1);
+			}
+		}
+	}
+	return FIndex3i(-1, -1, -1);
+}
+
+
+static bool QuadConnectedToTri(const FDynamicMesh3& Mesh, int32 TriangleID, FIndex2i Quad2)
+{
+	FIndex3i TriNbrsTris = Mesh.GetTriNeighbourTris(TriangleID);
+	if (IndexUtil::FindTriIndex(Quad2.A, TriNbrsTris) != IndexConstants::InvalidID ||
+		IndexUtil::FindTriIndex(Quad2.B, TriNbrsTris) != IndexConstants::InvalidID)
+	{
+		return true;
+	}
+	return false;
+}
+
+
+
+void UE::GroomQueries::ExtractCardQuads(
+	const FDynamicMesh3& Mesh,
+	UE::GroomQueries::FMeshCardStripSet& CardsInOut)
+{
+	FDynamicMesh3 TmpMesh;
+	TmpMesh.Copy(Mesh, false, false, false, false);
+
+	int32 NumCards = CardsInOut.CardStrips.Num();
+	for (int32 ci = 0; ci < NumCards; ++ci)
+	{
+		FMeshCardStrip& CardStrip = CardsInOut.CardStrips[ci];
+		bool bIsValidCard = true;
+
+		TArray<int32> RemainingTris(CardStrip.Triangles);
+		while (RemainingTris.Num() > 1)
+		{
+			bool bIsLastQuad = (RemainingTris.Num() < 4);
+			int32 LastRemainingTrisCount = RemainingTris.Num();
+
+			for (int32 tid : RemainingTris)
+			{
+				FIndex3i TriEdges = TmpMesh.GetTriEdges(tid);
+				FIndex3i TriBdryEdgeFlag(
+					TmpMesh.IsBoundaryEdge(TriEdges.A) ? 1 : 0,
+					TmpMesh.IsBoundaryEdge(TriEdges.B) ? 1 : 0,
+					TmpMesh.IsBoundaryEdge(TriEdges.C) ? 1 : 0);
+				bool bFoundMatch = false;
+				if ( (TriBdryEdgeFlag.A + TriBdryEdgeFlag.B + TriBdryEdgeFlag.C) == 2 )
+				{
+					int32 SharedEdgeIdx = IndexUtil::FindTriIndex(0, TriBdryEdgeFlag);
+					int32 SharedEdge = TriEdges[SharedEdgeIdx];
+					FIndex2i NbrEdgeT = TmpMesh.GetEdgeT(SharedEdge);
+					check(NbrEdgeT.B != IndexConstants::InvalidID);
+					int32 OtherTID = (NbrEdgeT.A == tid) ? NbrEdgeT.B : NbrEdgeT.A;
+
+					FIndex2i CardTris(tid, OtherTID);
+
+					FIndex3i Tri1 = TmpMesh.GetTriangle(CardTris.A);
+					FIndex3i Tri2 = TmpMesh.GetTriangle(CardTris.B);
+
+					// can we use this to sort out the correct cycle shift
+					FIndex2i SharedEdgeV = TmpMesh.GetEdgeV(SharedEdge);
+					FIndex2i OpposingV = TmpMesh.GetEdgeOpposingV(SharedEdge);
+
+					// need to cycle the first triangle verts correctly so that when we add fourth vertex we don't create an hourglass quad
+					FIndex4i Quad(Tri1.A, Tri1.B, Tri1.C, 0);		// SharedEdgeIdx == 2 case
+					if (SharedEdgeIdx == 0)
+					{
+						Quad = FIndex4i(Tri1.B, Tri1.C, Tri1.A, 0);
+					}
+					else if (SharedEdgeIdx == 1)
+					{
+						Quad = FIndex4i(Tri1.C, Tri1.A, Tri1.B, 0);
+					}
+
+					// figure out the fourth vertex, ie which of the vertices of tri2 was not in tri1
+					if (IndexUtil::FindTriIndex(Tri2.A, Tri1) == FDynamicMesh3::InvalidID)
+					{
+						Quad.D = Tri2.A;
+					}
+					else if (IndexUtil::FindTriIndex(Tri2.B, Tri1) == FDynamicMesh3::InvalidID)
+					{
+						Quad.D = Tri2.B;
+					}
+					else
+					{
+						Quad.D = Tri2.C;
+					}
+					
+					int32 FreeVertIdx = -1;
+					for (int32 j = 0; j < 3 && FreeVertIdx == -1; ++j)
+					{
+						FreeVertIdx = (IndexUtil::FindTriIndex(Quad[j], Tri2) == FDynamicMesh3::InvalidID) ? j : -1;
+					}
+
+					// attach new quad to start or end of strip depending on which tris were adjacent (does this work??)
+					if (CardStrip.QuadTriPairs.Num() == 0 || QuadConnectedToTri(Mesh, CardStrip.QuadTriPairs.Last().B, CardTris))
+					{
+						CardStrip.QuadTriPairs.Add(CardTris);
+						CardStrip.QuadLoops.Add(Quad);
+					}
+					else if (QuadConnectedToTri(Mesh, CardStrip.QuadTriPairs[0].A, CardTris))
+					{
+						CardStrip.QuadTriPairs.Insert(CardTris, 0);
+						CardStrip.QuadLoops.Insert(Quad, 0);
+					}
+					else
+					{
+						check(false);
+					}
+
+					RemainingTris.RemoveSwap(tid, false);
+					RemainingTris.RemoveSwap(OtherTID, false);
+					TmpMesh.RemoveTriangle(tid, false);
+
+					FIndex3i OtherNbrTris = TmpMesh.GetTriNeighbourTris(OtherTID);
+					TmpMesh.RemoveTriangle(OtherTID, false);
+
+					// figure out the next connected tri in the strip, and reorder array so that we start there next iteration.
+					// this allows the quads to be processed in-order
+					for (int32 j = 0; j < 3; ++j)
+					{
+						if (OtherNbrTris[j] != IndexConstants::InvalidID && OtherNbrTris[j] != tid && RemainingTris.Contains(OtherNbrTris[j]))
+						{
+							int32 OtherIdx = RemainingTris.IndexOfByKey(OtherNbrTris[j]);
+							Swap(RemainingTris[0], RemainingTris[OtherIdx]);
+						}
+					}
+					bFoundMatch = true;
+				}
+
+				if (bFoundMatch)		// break out of loop over RemainingTris and go to next
+				{
+					break;
+				}
+			}
+
+			// if we did not remove any triangles this loop, this is not a valid card
+			if (LastRemainingTrisCount == RemainingTris.Num())
+			{
+				bIsValidCard = false;
+				break;
+			}
+		}
+
+		if (!bIsValidCard)
+		{
+			continue;
+		}
+
+		// Output of code above has each quad oriented properly, but the vertex cycle is not consistent.
+		// So now we walk down the strip and orient each quad in order [A,B,C,D] such that the centerline
+		// of the strip passes through Midpoint(A,B) and Midpoint(C,D), and the "side" edges are (B,C) and (D,A).
+		// We can figure out the correct ordering based on which edge is currently shared w/ the next Quad,
+		// except at the last quad, which has no next so we have to use the Previous
+		int32 NumQuads = CardStrip.QuadLoops.Num();
+		for (int32 qi = 0; qi < NumQuads && NumQuads > 1; ++qi)
+		{
+			bool bIsLast = (qi == NumQuads - 1);
+			FIndex4i CurQuad = CardStrip.QuadLoops[qi];
+			FIndex4i NextQuad = CardStrip.QuadLoops[bIsLast ? (NumQuads-2) : (qi+1)];		// use previous for last quad
+			FIndex3i SharedEdgeInfo = FindSharedQuadEdge(Mesh, CurQuad, NextQuad);
+			if (ensure(SharedEdgeInfo.A != -1))
+			{
+				SharedEdgeInfo.A = (bIsLast) ? ((SharedEdgeInfo.A + 2) % 4) : SharedEdgeInfo.A;		// cases are reversed for last quad
+				if (SharedEdgeInfo.A == 0)
+				{
+					CardStrip.QuadLoops[qi] = FIndex4i(CurQuad.C, CurQuad.D, CurQuad.A, CurQuad.B);
+				}
+				else if (SharedEdgeInfo.A == 1)
+				{
+					CardStrip.QuadLoops[qi] = FIndex4i(CurQuad.D, CurQuad.A, CurQuad.B, CurQuad.C);
+				}
+				else if (SharedEdgeInfo.A == 3)
+				{
+					CardStrip.QuadLoops[qi] = FIndex4i(CurQuad.B, CurQuad.C, CurQuad.D, CurQuad.A);
+				}
+			}
+		}
+
+	}
+}
+
+
+
+void UE::GroomQueries::ExtractCardCurves(
+	const FDynamicMesh3& Mesh,
+	UE::GroomQueries::FMeshCardStripSet& CardsInOut)
+{
+
+	int32 NumCards = CardsInOut.CardStrips.Num();
+	for (int32 ci = 0; ci < NumCards; ++ci)
+	{
+		FMeshCardStrip& CardStrip = CardsInOut.CardStrips[ci];
+
+		// skip strips w/ no quads
+		if (CardStrip.QuadLoops.Num() == 0)
+		{
+			continue;
+		}
+
+		//for (FIndex4i Quad : CardStrip.QuadLoops)
+		//{
+		//	FVector3d Center =
+		//		Mesh.GetVertex(Quad.A) + Mesh.GetVertex(Quad.B) + Mesh.GetVertex(Quad.C) + Mesh.GetVertex(Quad.D);
+		//	Center /= 4.0;
+		//	CardStrip.CardCurve.Add(Center);
+		//	CardStrip.CardCurveWidth.Add(0.0);
+		//}
+
+		int32 NumQuads = CardStrip.QuadLoops.Num();
+		for (int32 k = 0; k <= NumQuads; ++k)
+		{
+			FVector3d A, B;
+			if (k == NumQuads)
+			{
+				FIndex4i CardQuad = CardStrip.QuadLoops[k-1];
+				A = Mesh.GetVertex(CardQuad.C);
+				B = Mesh.GetVertex(CardQuad.D);
+			}
+			else
+			{
+				FIndex4i CardQuad = CardStrip.QuadLoops[k];
+				A = Mesh.GetVertex(CardQuad.A);
+				B = Mesh.GetVertex(CardQuad.B);
+			}
+			CardStrip.CardCurve.Add( (A+B) * 0.5 );
+			CardStrip.CardCurveWidth.Add((B - A).Length() * 0.5);
+		}
+
+
+		int32 NumVertices = CardStrip.CardCurve.Num();
+		for (int32 k = 0; k < NumVertices; ++k)
+		{
+			FIndex2i AdjacentVerts = CardStrip.GetCurvePointVertices(k);
+			FVector3d Left = Mesh.GetVertex(AdjacentVerts.A);
+			FVector3d Right = Mesh.GetVertex(AdjacentVerts.B);
+			FVector3d Center = CardStrip.CardCurve[k];
+
+			FFrame3d CardFrame(Center);
+			CardFrame.AlignAxis(0, UE::Geometry::Normalized(Right - Center));
+
+			FVector3d Y;
+			if (k == 0)
+			{
+				FVector3d Forward = CardStrip.CardCurve[k + 1];
+				Y = UE::Geometry::Normalized(Forward - Center);
+			}
+			else if (k == NumVertices - 1)
+			{
+				FVector3d Back = CardStrip.CardCurve[k-1];
+				Y = UE::Geometry::Normalized(Center - Back);
+			}
+			else
+			{
+				FVector3d Forward = CardStrip.CardCurve[k + 1];
+				FVector3d Back = CardStrip.CardCurve[k - 1];
+				Y = UE::Geometry::Normalized(Forward - Back);
+			}
+
+			CardFrame.ConstrainedAlignAxis(1, Y, CardFrame.X());
+			CardStrip.CardCurveFrames.Add(CardFrame);
+		}
+	}
+
+}
+
+
+
+
+
+
+
+
+
+void UE::GroomQueries::ProcessHairCurves(AGroomActor* GroomActor,
+	bool bUseGuides,
+	TFunctionRef<void(const TArrayView<FVector3f>& Positions, const TArrayView<float>& Radii)> HairCurveFunc)
+{
+	check(GroomActor->GetGroomComponent());
+	check(GroomActor->GetGroomComponent()->GroomAsset);
+	UGroomAsset* Asset = GroomActor->GetGroomComponent()->GroomAsset;
+
+	int32 NumHairGroups = Asset->HairGroupsInfo.Num();
+	for (int32 GroupIdx = 0; GroupIdx < NumHairGroups; ++GroupIdx)
+	{
+		const FHairGroupInfo& GroupInfo = Asset->HairGroupsInfo[GroupIdx];
+		const FHairGroupData& GroupData = Asset->HairGroupsData[GroupIdx];
+
+		//int32 NumCurves = GroupInfo.NumCurves;
+		//int32 NumGuides = GroupInfo.NumGuides;
+
+		FHairStrandsDatas StrandsData;
+		FHairStrandsDatas GuidesData;
+		Asset->GetHairStrandsDatas(GroupIdx, StrandsData, GuidesData);
+
+		//const FHairStrandsDatas& GroupStrandData = (bUseGuides) ? GroupData.HairSimulationData : GroupData.HairRenderData;
+		const FHairStrandsDatas& GroupStrandData = (bUseGuides) ? GuidesData : StrandsData;
+
+		const FHairStrandsPoints& GroupStrandPoints = GroupStrandData.StrandsPoints;
+		const TArray<FVector3f>& Positions = GroupStrandPoints.PointsPosition;
+		const TArray<float>& Radii = GroupStrandPoints.PointsRadius;
+
+		const FHairStrandsCurves& GroupStrandCurves = GroupStrandData.StrandsCurves;
+		const TArray<uint16>& CurvesCounts = GroupStrandCurves.CurvesCount;
+		const TArray<uint32>& CurvesOffsets = GroupStrandCurves.CurvesOffset;
+
+		int32 NumCurves = CurvesCounts.Num();
+		for (int32 CurveIndex = 0; CurveIndex < NumCurves; ++CurveIndex)
+		{
+			int32 Count = (int32)CurvesCounts[CurveIndex];
+			int32 Offset = (int32)CurvesOffsets[CurveIndex];
+
+			TArrayView<FVector3f> CurvePositions = TArrayView<FVector3f>((FVector3f*)&Positions[Offset], Count);
+			TArrayView<float> CurveRadii = TArrayView<float>((float*)&Radii[Offset], Count);
+
+			HairCurveFunc(CurvePositions, CurveRadii);
+
+			//for (int32 k = 0; k < Count - 1; ++k)
+			//{
+			//	FSkeletalImplicitLine3d Line = { FSegment3d(FVector3d(Positions[Offset + k]), FVector3d(Positions[Offset + k + 1])), Radii[Offset + k] };
+			//	//Line.Radius = 5.0;
+			//	Lines.Add(Line);
+			//}
+		}
+
+		//FString LinesName = FString::Printf(TEXT("STRANDGROUP_%d"), GroupIdx);
+		//PreviewGeom->CreateOrUpdateLineSet(LinesName, GroupStrandCurves.Num(),
+		//	[&](int32 CurveIndex, TArray<FRenderableLine>& LinesOut)
+		//{
+		//	int32 Count = (int32)CurvesCounts[CurveIndex];
+		//	int32 Offset = (int32)CurvesOffsets[CurveIndex];
+		//	for (int32 k = 0; k < Count - 1; ++k)
+		//	{
+		//		LinesOut.Add(FRenderableLine(Positions[Offset + k], Positions[Offset + k + 1], FColor::Red, 0.1f));
+		//	}
+		//});
+	}
+
+}
\ No newline at end of file
diff --git a/Engine/Plugins/Experimental/HairModelingToolset/Source/HairModelingToolset/Private/GroomQueryUtil.h b/Engine/Plugins/Experimental/HairModelingToolset/Source/HairModelingToolset/Private/GroomQueryUtil.h
new file mode 100644
--- /dev/null
+++ b/Engine/Plugins/Experimental/HairModelingToolset/Source/HairModelingToolset/Private/GroomQueryUtil.h
@@ -0,0 +1,100 @@
+// Copyright Epic Games, Inc. All Rights Reserved.
+
+#pragma once
+
+#include "CoreMinimal.h"
+
+#include "GroomActor.h"
+#include "GroomComponent.h"
+#include "GroomAsset.h"
+
+#include "DynamicMesh/DynamicMesh3.h"
+
+
+
+namespace UE
+{
+	namespace GroomQueries
+	{
+		using namespace UE::Geometry;
+
+		// source card information
+		struct FCardSourceInfo
+		{
+			int32 GroupIndex;		// groom asset group
+			int32 LODIndex;			// groom asset LOD
+			int32 CardIndex;		// index of card into card arrays
+			int32 IndexOffset;		// index offset into card triangle buffer
+			int32 IndexCount;		// number of indices after offset that belong to this card
+		};
+
+		struct FMeshCardStrip
+		{
+			FCardSourceInfo SourceInfo;
+
+			int32 GroupID;
+			TArray<int32> Triangles;
+
+			TArray<FIndex2i> QuadTriPairs;
+			TArray<FIndex4i> QuadLoops;
+
+			TArray<FVector3d> CardCurve;
+			TArray<FFrame3d> CardCurveFrames;
+			TArray<double> CardCurveWidth;
+
+			FIndex2i GetCurvePointVertices(int32 Index) const
+			{
+				if (Index >= (CardCurve.Num() - 1))
+				{
+					Index = QuadLoops.Num() - 1;
+					return FIndex2i(QuadLoops[Index].D, QuadLoops[Index].C);
+				}
+				else
+				{
+					return FIndex2i(QuadLoops[Index].A, QuadLoops[Index].B);
+				}
+			}
+		};
+
+		struct FMeshCardStripSet
+		{
+			TArray<FMeshCardStrip> CardStrips;
+
+			// map from mesh triangle GroupIDs to CardStrips array index
+			TMap<int32,int32> TriGroupToCardIndex;
+
+			const FMeshCardStrip& FindStripForGroup(int32 GroupIdx) const;
+			FMeshCardStrip& FindStripForGroup(int32 GroupIdx);
+		};
+
+		void ExtractAllHairCards(AGroomActor* GroomActor,
+			int32 LODIndex,
+			FDynamicMesh3& MeshOut,
+			FMeshCardStripSet* CardsOut = nullptr);
+
+		void ExtractAllHairCards(
+			const FMeshDescription* SourceMeshIn,
+			FDynamicMesh3& MeshOut,
+			FMeshCardStripSet& CardsOut);
+
+		void ExtractCardQuads(
+			const FDynamicMesh3& Mesh,
+			FMeshCardStripSet& CardsInOut);
+
+		void ExtractCardCurves(
+			const FDynamicMesh3& Mesh,
+			FMeshCardStripSet& CardsInOut);
+
+		void ProcessHairCurves(AGroomActor* GroomActor,
+			bool bUseGuides,
+			TFunctionRef<void(const TArrayView<FVector3f>& Positions, const TArrayView<float>& Radii)> HairCurveFunc);
+
+
+	}
+}
+
+
+
+
+
+
diff --git a/Engine/Plugins/Experimental/HairModelingToolset/Source/HairModelingToolset/Private/GroomToMeshTool.cpp b/Engine/Plugins/Experimental/HairModelingToolset/Source/HairModelingToolset/Private/GroomToMeshTool.cpp
new file mode 100644
--- /dev/null
+++ b/Engine/Plugins/Experimental/HairModelingToolset/Source/HairModelingToolset/Private/GroomToMeshTool.cpp
@@ -0,0 +1,1341 @@
+// Copyright Epic Games, Inc. All Rights Reserved.
+
+#include "GroomToMeshTool.h"
+#include "InteractiveToolManager.h"
+#include "ToolBuilderUtil.h"
+
+#include "DynamicMesh/DynamicMesh3.h"
+#include "DynamicMesh/DynamicMeshAABBTree3.h"
+#include "Spatial/FastWinding.h"
+#include "DynamicMesh/MeshNormals.h"
+#include "DynamicMeshEditor.h"
+#include "Util/ColorConstants.h"
+#include "ToolSetupUtil.h"
+#include "Selection/ToolSelectionUtil.h"
+#include "ModelingObjectsCreationAPI.h"
+#include "DynamicMesh/MeshTransforms.h"
+#include "Selections/MeshFaceSelection.h"
+#include "Selections/MeshVertexSelection.h"
+#include "Selections/MeshConnectedComponents.h"
+#include "DynamicMeshEditor.h"
+#include "MeshBoundaryLoops.h"
+#include "Parameterization/DynamicMeshUVEditor.h"
+#include "Parameterization/MeshUVTransforms.h"
+
+#include "DynamicMesh/DynamicPointSet3.h"
+#include "DynamicMesh/DynamicVerticesOctree3.h"
+#include "DynamicMesh/DynamicAttribute.h"
+#include "DynamicMesh/DynamicVertexAttribute.h"
+
+#include "Implicit/CachingMeshSDF.h"
+#include "Implicit/GridInterpolant.h"
+#include "Implicit/ImplicitFunctions.h"
+#include "Implicit/Morphology.h"
+#include "Generators/MarchingCubes.h"
+
+#include "MeshSimplification.h"
+#include "SmoothingOps/CotanSmoothingOp.h"
+
+#include "GroomActor.h"
+#include "GroomComponent.h"
+#include "GroomAsset.h"
+#include "HairDescription.h"
+
+#include "ParameterizationOps/UVLayoutOp.h"
+
+#include "Engine/Classes/Engine/StaticMesh.h"
+#include "Engine/Classes/Components/StaticMeshComponent.h"
+#include "Engine/Classes/Engine/StaticMeshActor.h"
+#include "MeshDescriptionToDynamicMesh.h"
+
+#include "ExplicitUseGeometryMathTypes.h"		// using UE::Geometry::(math types)
+using namespace UE::Geometry;
+
+#define LOCTEXT_NAMESPACE "UGroomToMeshTool"
+
+
+/*
+ * ToolBuilder
+ */
+
+bool UGroomToMeshToolBuilder::CanBuildTool(const FToolBuilderState& SceneState) const
+{
+	return ToolBuilderUtil::CountSelectedActorsOfType<AGroomActor>(SceneState) == 1;
+}
+
+UInteractiveTool* UGroomToMeshToolBuilder::BuildTool(const FToolBuilderState& SceneState) const
+{
+	UGroomToMeshTool* NewTool = NewObject<UGroomToMeshTool>(SceneState.ToolManager);
+
+	NewTool->SetWorld(SceneState.World);
+	AGroomActor* Groom = ToolBuilderUtil::FindFirstActorOfType<AGroomActor>(SceneState);
+	check(Groom != nullptr);
+	NewTool->SetSelection(Groom);
+
+	return NewTool;
+}
+
+
+
+
+/*
+ * Tool
+ */
+UGroomToMeshTool::UGroomToMeshTool()
+{
+	SetToolDisplayName(LOCTEXT("GroomToMeshToolName", "Groom to Mesh"));
+}
+
+
+void UGroomToMeshTool::SetSelection(AGroomActor* Groom)
+{
+	TargetGroom = Groom;
+}
+
+
+void UGroomToMeshTool::Setup()
+{
+	UInteractiveTool::Setup();
+
+	PreviewMesh = NewObject<UPreviewMesh>(this);
+	PreviewMesh->bBuildSpatialDataStructure = false;
+	PreviewMesh->CreateInWorld(TargetGroom->GetWorld(), FTransform::Identity);
+	PreviewMesh->SetTransform(TargetGroom->GetActorTransform());
+
+	MeshMaterial = ToolSetupUtil::GetDefaultSculptMaterial(GetToolManager());
+	UVMaterial = ToolSetupUtil::GetUVCheckerboardMaterial(50.0);
+	PreviewMesh->SetMaterial(MeshMaterial);
+
+	PreviewGeom = NewObject<UPreviewGeometry>(this);
+	PreviewGeom->CreateInWorld(TargetGroom->GetWorld(), FTransform::Identity);
+	PreviewGeom->GetActor()->SetActorTransform(TargetGroom->GetActorTransform());
+
+	Settings = NewObject<UGroomToMeshToolProperties>(this);
+	Settings->RestoreProperties(this);
+	AddToolPropertySource(Settings);
+
+	Settings->WatchProperty(Settings->VoxelCount, [this](int32) { bResultValid = false; });
+	Settings->WatchProperty(Settings->BlendPower, [this](float) { bResultValid = false; });
+	Settings->WatchProperty(Settings->RadiusScale, [this](float) { bResultValid = false; });
+
+	Settings->WatchProperty(Settings->bApplyMorphology, [this](bool) { bResultValid = false; });
+	Settings->WatchProperty(Settings->ClosingDist, [this](float) { bResultValid = false; });
+	Settings->WatchProperty(Settings->OpeningDist, [this](float) { bResultValid = false; });
+
+	Settings->WatchProperty(Settings->bClipToHead, [this](bool) { bResultValid = false; });
+	Settings->WatchProperty(Settings->ClipMeshActor, [this](TLazyObjectPtr<AStaticMeshActor>) { bResultValid = false; });
+
+	Settings->WatchProperty(Settings->bSmooth, [this](bool) { bResultValid = false; });
+	Settings->WatchProperty(Settings->Smoothness, [this](float) { bResultValid = false; });
+	Settings->WatchProperty(Settings->VolumeCorrection, [this](float) { bResultValid = false; });
+
+	Settings->WatchProperty(Settings->UVMode, [this](EGroomToMeshUVMode) { bResultValid = false; });
+
+	Settings->WatchProperty(Settings->bSimplify, [this](bool) { bResultValid = false; });
+	Settings->WatchProperty(Settings->VertexCount, [this](int32) { bResultValid = false; });
+	
+	Settings->WatchProperty(Settings->bShowSideBySide, [this](bool) { bResultValid = false; });
+	Settings->WatchProperty(Settings->bShowGuides, [this](bool) { bResultValid = false; });
+	Settings->WatchProperty(Settings->bShowUVs, [this](bool) { bVisualizationChanged = true; });
+
+	bResultValid = false;
+	bVisualizationChanged = true;
+
+	GetToolManager()->DisplayMessage( 
+		LOCTEXT("OnStartTool", "Convert a Groom to a Static Mesh"),
+		EToolMessageLevel::UserNotification);
+}
+
+
+
+void UGroomToMeshTool::Shutdown(EToolShutdownType ShutdownType)
+{
+	Settings->SaveProperties(this);
+
+	UE::Geometry::FTransform3d Transform(PreviewMesh->GetTransform());
+	PreviewMesh->Disconnect();
+	PreviewMesh = nullptr;
+
+	PreviewGeom->Disconnect();
+	PreviewGeom = nullptr;
+
+	if (ShutdownType == EToolShutdownType::Accept )
+	{
+		UMaterialInterface* UseMaterial = UMaterial::GetDefaultMaterial(MD_Surface);
+
+		FString NewName = TargetGroom.IsValid() ?
+			FString::Printf(TEXT("%sMesh"), *TargetGroom->GetName()) : TEXT("Groom Mesh");
+
+		GetToolManager()->BeginUndoTransaction(LOCTEXT("CreateMeshGroom", "Groom To Mesh"));
+
+		FCreateMeshObjectParams NewMeshObjectParams;
+		NewMeshObjectParams.TargetWorld = TargetWorld;
+		NewMeshObjectParams.Transform = (FTransform)Transform;
+		NewMeshObjectParams.BaseName = NewName;
+		NewMeshObjectParams.Materials.Add(UseMaterial);
+		NewMeshObjectParams.SetMesh(&CurrentMesh);
+		FCreateMeshObjectResult Result = UE::Modeling::CreateMeshObject(GetToolManager(), MoveTemp(NewMeshObjectParams));
+		if (Result.IsOK() && Result.NewActor != nullptr)
+		{
+			ToolSelectionUtil::SetNewActorSelection(GetToolManager(), Result.NewActor);
+		}
+
+		GetToolManager()->EndUndoTransaction();
+	}
+
+
+}
+
+
+void UGroomToMeshTool::OnTick(float DeltaTime)
+{
+	if (bResultValid == false)
+	{
+		RecalculateMesh();
+		bResultValid = true;
+	}
+
+	if (bVisualizationChanged)
+	{
+		PreviewMesh->SetMaterial(Settings->bShowUVs ? UVMaterial : MeshMaterial);
+		bVisualizationChanged = false;
+	}
+}
+
+void UGroomToMeshTool::Render(IToolsContextRenderAPI* RenderAPI)
+{
+}
+
+
+bool UGroomToMeshTool::HasAccept() const
+{
+	return true;
+}
+
+bool UGroomToMeshTool::CanAccept() const
+{
+	return true;
+}
+
+
+
+
+
+class FHairPointSetAttributes : public TDynamicAttributeSetBase<FDynamicPointSet3d>
+{
+public:
+	using FRadiusAttribute = TDynamicVertexAttribute<float, 1, FDynamicPointSet3d>;
+
+	TUniquePtr<FRadiusAttribute> Radius;
+
+	void Initialize(FDynamicPointSet3d* Parent)
+	{
+		Radius = MakeUnique<FRadiusAttribute>(Parent);
+
+		RegisterExternalAttribute(Radius.Get());
+	}
+
+};
+
+
+
+void UGroomToMeshTool::RecalculateMesh()
+{
+	TSharedPtr<FDynamicMesh3, ESPMode::ThreadSafe> CurrentResult = UpdateVoxelization();
+
+	// helper to only use NewMesh if it was valid, otherwise ignore it
+	auto SelectValidMesh = [](TSharedPtr<FDynamicMesh3, ESPMode::ThreadSafe> InCurrentMesh, TSharedPtr<FDynamicMesh3, ESPMode::ThreadSafe> InNewMesh)
+	{
+		if (InNewMesh->TriangleCount() > 0)
+		{
+			return InNewMesh;
+		}
+		return InCurrentMesh;
+	};
+
+	if (Settings->bApplyMorphology && CurrentResult->TriangleCount() > 0)
+	{
+		CurrentResult = SelectValidMesh(CurrentResult, UpdateMorphology(CurrentResult));
+	}
+
+	if (Settings->bClipToHead && CurrentResult->TriangleCount() > 0)
+	{
+		CurrentResult = SelectValidMesh(CurrentResult, UpdateClipMesh(CurrentResult));
+	}
+
+	if (Settings->bSmooth && CurrentResult->TriangleCount() > 0)
+	{
+		CurrentResult = SelectValidMesh(CurrentResult, UpdateSmoothing(CurrentResult));
+	}
+
+	if (Settings->bSimplify && CurrentResult->TriangleCount() > 0)
+	{
+		CurrentResult = SelectValidMesh(CurrentResult, UpdateSimplification(CurrentResult));
+	}
+
+	if (CurrentResult->TriangleCount() > 0)
+	{
+		CurrentResult = SelectValidMesh(CurrentResult, UpdatePostprocessing(CurrentResult));
+	}
+
+	UpdatePreview(CurrentResult);
+}
+
+
+
+
+
+
+
+
+static void ProcessHairCurvePoints(AGroomActor* GroomActor,
+	bool bUseGuides,
+	TFunctionRef<void(const TArray<FVector3f>& Positions, const TArray<float>& Radii)> HairCurvePointsFunc)
+{
+	check(GroomActor->GetGroomComponent());
+	check(GroomActor->GetGroomComponent()->GroomAsset);
+	UGroomAsset* Asset = GroomActor->GetGroomComponent()->GroomAsset;
+
+	int32 NumHairGroups = Asset->HairGroupsInfo.Num();
+	for (int32 GroupIdx = 0; GroupIdx < NumHairGroups; ++GroupIdx)
+	{
+		const FHairGroupInfo& GroupInfo = Asset->HairGroupsInfo[GroupIdx];
+		const FHairGroupData& GroupData = Asset->HairGroupsData[GroupIdx];
+
+		//int32 NumCurves = GroupInfo.NumCurves;
+		//int32 NumGuides = GroupInfo.NumGuides;
+
+		//const FHairStrandsDatas& GroupStrandData = (bUseGuides) ? GroupData.HairSimulationData : GroupData.HairRenderData;
+		FHairStrandsDatas StrandsData;
+		FHairStrandsDatas GuidesData;
+		Asset->GetHairStrandsDatas(GroupIdx, StrandsData, GuidesData);
+		const FHairStrandsDatas& GroupStrandData = (bUseGuides) ? GuidesData : StrandsData;
+
+		const FHairStrandsPoints& GroupStrandPoints = GroupStrandData.StrandsPoints;
+		const TArray<FVector3f>& Positions = GroupStrandPoints.PointsPosition;
+		const TArray<float>& Radii = GroupStrandPoints.PointsRadius;
+
+		HairCurvePointsFunc(GroupStrandPoints.PointsPosition, GroupStrandPoints.PointsRadius);
+
+		//int32 NumPoints = Positions.Num();
+		//for (int32 k = 0; k < NumPoints; ++k)
+		//{
+		//	FVector3d Position(Positions[k]);
+		//	Bounds.Contain(Position);
+		//	int32 NewVID = Points->AppendVertex(Position);
+		//	PointAttribs->Radius->SetNewValue(NewVID, &Radii[k]);
+		//}
+
+	}
+}
+
+
+
+
+
+
+static void ProcessHairCurves(AGroomActor* GroomActor,
+	bool bUseGuides,
+	TFunctionRef<void(const TArrayView<FVector3f>& Positions, const TArrayView<float>& Radii)> HairCurveFunc)
+{
+	check(GroomActor->GetGroomComponent());
+	check(GroomActor->GetGroomComponent()->GroomAsset);
+	UGroomAsset* Asset = GroomActor->GetGroomComponent()->GroomAsset;
+
+	int32 NumHairGroups = Asset->HairGroupsInfo.Num();
+	for (int32 GroupIdx = 0; GroupIdx < NumHairGroups; ++GroupIdx)
+	{
+		const FHairGroupInfo& GroupInfo = Asset->HairGroupsInfo[GroupIdx];
+		const FHairGroupData& GroupData = Asset->HairGroupsData[GroupIdx];
+
+		//int32 NumCurves = GroupInfo.NumCurves;
+		//int32 NumGuides = GroupInfo.NumGuides;
+
+		FHairStrandsDatas StrandsData;
+		FHairStrandsDatas GuidesData;
+		Asset->GetHairStrandsDatas(GroupIdx, StrandsData, GuidesData);
+		//const FHairStrandsDatas& GroupStrandData = (bUseGuides) ? GroupData.HairSimulationData : GroupData.HairRenderData;
+		const FHairStrandsDatas& GroupStrandData = (bUseGuides) ? GuidesData : StrandsData;
+
+		const FHairStrandsPoints& GroupStrandPoints = GroupStrandData.StrandsPoints;
+		const TArray<FVector3f>& Positions = GroupStrandPoints.PointsPosition;
+		const TArray<float>& Radii = GroupStrandPoints.PointsRadius;
+
+		const FHairStrandsCurves& GroupStrandCurves = GroupStrandData.StrandsCurves;
+		const TArray<uint16>& CurvesCounts = GroupStrandCurves.CurvesCount;
+		const TArray<uint32>& CurvesOffsets = GroupStrandCurves.CurvesOffset;
+
+		int32 NumCurves = CurvesCounts.Num();
+		for (int32 CurveIndex = 0; CurveIndex < NumCurves; ++CurveIndex)
+		{
+			int32 Count = (int32)CurvesCounts[CurveIndex];
+			int32 Offset = (int32)CurvesOffsets[CurveIndex];
+
+			TArrayView<FVector3f> CurvePositions = TArrayView<FVector3f>( (FVector3f*)&Positions[Offset], Count );
+			TArrayView<float> CurveRadii = TArrayView<float>((float*)&Radii[Offset], Count);
+
+			HairCurveFunc(CurvePositions, CurveRadii);
+
+			//for (int32 k = 0; k < Count - 1; ++k)
+			//{
+			//	FSkeletalImplicitLine3d Line = { FSegment3d(FVector3d(Positions[Offset + k]), FVector3d(Positions[Offset + k + 1])), Radii[Offset + k] };
+			//	//Line.Radius = 5.0;
+			//	Lines.Add(Line);
+			//}
+		}
+
+		//FString LinesName = FString::Printf(TEXT("STRANDGROUP_%d"), GroupIdx);
+		//PreviewGeom->CreateOrUpdateLineSet(LinesName, GroupStrandCurves.Num(),
+		//	[&](int32 CurveIndex, TArray<FRenderableLine>& LinesOut)
+		//{
+		//	int32 Count = (int32)CurvesCounts[CurveIndex];
+		//	int32 Offset = (int32)CurvesOffsets[CurveIndex];
+		//	for (int32 k = 0; k < Count - 1; ++k)
+		//	{
+		//		LinesOut.Add(FRenderableLine(Positions[Offset + k], Positions[Offset + k + 1], FColor::Red, 0.1f));
+		//	}
+		//});
+	}
+
+}
+
+
+
+
+
+TSharedPtr<FDynamicMesh3, ESPMode::ThreadSafe> UGroomToMeshTool::UpdateVoxelization()
+{
+	FVoxelizeSettings NewSettings;
+	NewSettings.VoxelCount = Settings->VoxelCount;
+	NewSettings.BlendPower = Settings->BlendPower;
+	NewSettings.RadiusScale = Settings->RadiusScale;
+	if (CachedVoxelizeSettings == NewSettings)
+	{
+		return CurrentVoxelizeResult;
+	}
+
+	check(TargetGroom.IsValid());
+	AGroomActor* GroomActor = TargetGroom.Get();
+
+	TUniquePtr<FDynamicPointSet3d> Points = MakeUnique<FDynamicPointSet3d>();
+	TUniquePtr<FHairPointSetAttributes> PointAttribs = MakeUnique<FHairPointSetAttributes>();
+	PointAttribs->Initialize(Points.Get());
+	Points->SetExternallyManagedAttributes(PointAttribs.Get());
+
+	FAxisAlignedBox3d Bounds = FAxisAlignedBox3d::Empty();
+	ProcessHairCurvePoints(GroomActor, true, [&](const TArray<FVector3f>& Positions, const TArray<float>& Radii)
+	{
+		int32 NumPoints = Positions.Num();
+		for (int32 k = 0; k < NumPoints; ++k)
+		{
+			FVector3d Position(Positions[k]);
+			Bounds.Contain(Position);
+			int32 NewVID = Points->AppendVertex(Position);
+			PointAttribs->Radius->SetNewValue(NewVID, &Radii[k]);
+		}
+	});
+
+
+	//TDynamicVerticesOctree3<FDynamicPointSet3d> PointOctree;
+	//PointOctree.Initialize(Points.Get(), true);
+
+
+	double BlendPower = NewSettings.BlendPower;
+	int32 TargetVoxelCount = NewSettings.VoxelCount;
+	double MaxRadius = 0.0;
+	double RadiusScale = NewSettings.RadiusScale;
+
+	TArray<FSkeletalImplicitLine3d> Lines;
+
+	ProcessHairCurves(GroomActor, true, [&](const TArrayView<FVector3f>& Positions, const TArrayView<float>& Radii)
+	{
+		int32 Count = Positions.Num() - 1;
+		for (int32 k = 0; k < Count; ++k)
+		{
+			FSkeletalImplicitLine3d Line;
+			Line.Segment = FSegment3d(FVector3d(Positions[k]), FVector3d(Positions[k + 1]));
+			Line.SetScaleFromRadius(RadiusScale * Radii[k]);
+			Lines.Add(Line);
+		}
+	});
+
+	TSkeletalRicciNaryBlend3<FSkeletalImplicitLine3d, double> Blend;
+	Blend.BlendPower = BlendPower;
+	Blend.Children.Reserve(Lines.Num());
+	for (int32 k = 0; k < Lines.Num(); k++)
+	{
+		Blend.Children.Add(&Lines[k]);
+		MaxRadius = FMathd::Max(Lines[k].GetRadius(), MaxRadius);
+	}
+
+	FAxisAlignedBox3d UseBounds = Blend.Bounds();
+	FVector3d BoxDims = UseBounds.Diagonal();
+	double MaxDimension = UseBounds.MaxDim();
+
+	//MaxRadius = 10.0;		// why is this necessary??!?
+	double MeshCellSize = (UseBounds.MaxDim() + 2.0 * MaxRadius) / double(TargetVoxelCount);
+	//double MeshCellSize = UseBounds.MaxDim() / double(TargetVoxelCount);
+
+	FMarchingCubes MarchingCubes;
+	//MarchingCubes.CancelF = CancelF;
+
+	MarchingCubes.CubeSize = MeshCellSize;
+	MarchingCubes.IsoValue = 0.5f;
+	MarchingCubes.Bounds = UseBounds;
+	MarchingCubes.Bounds.Expand(2.0 * MaxRadius);
+	//MarchingCubes.Bounds.Expand(32.0 * MeshCellSize);
+	MarchingCubes.RootMode = ERootfindingModes::LerpSteps;
+	//MarchingCubes.RootMode = ERootfindingModes::Bisection;
+	MarchingCubes.RootModeSteps = 4;
+
+	TArray<FVector3d> Seeds;
+	for (int32 k = 0; k < Lines.Num(); ++k)
+	{
+		Seeds.Add( Lines[k].Segment.EndPoint() + Lines[k].GetRadius()*Lines[k].Segment.Direction );
+		Seeds.Add( Lines[k].Segment.StartPoint() - Lines[k].GetRadius()*Lines[k].Segment.Direction );
+	}
+
+	MarchingCubes.Implicit = [&Blend](const FVector3d& Pt) { return Blend.Value(Pt); };
+
+	MarchingCubes.GenerateContinuation(Seeds);
+
+	// if we found zero triangles, try again w/ full grid search
+	if (MarchingCubes.Triangles.Num() == 0)
+	{
+		MarchingCubes.Generate();
+	}
+
+	// clear implicit function...
+	MarchingCubes.Implicit = nullptr;
+
+	CurrentVoxelizeResult = MakeShared<FDynamicMesh3, ESPMode::ThreadSafe>(&MarchingCubes);
+	CachedVoxelizeSettings = NewSettings;
+	return CurrentVoxelizeResult;
+}
+
+
+
+
+TSharedPtr<FDynamicMesh3, ESPMode::ThreadSafe> UGroomToMeshTool::UpdateMorphology(TSharedPtr<FDynamicMesh3, ESPMode::ThreadSafe> InputMesh)
+{
+	FMorphologySettings NewSettings;
+	NewSettings.InputMesh = InputMesh;
+	NewSettings.VoxelCount = Settings->VoxelCount;
+	NewSettings.OpenDist = Settings->OpeningDist;
+	NewSettings.CloseDist = Settings->ClosingDist;
+	if (CachedMorphologySettings == NewSettings)
+	{
+		return CachedMorphologyResult;
+	}
+
+	CachedMorphologyResult = MakeShared<FDynamicMesh3, ESPMode::ThreadSafe>(*InputMesh);
+
+	if (NewSettings.CloseDist > 0 && CachedMorphologyResult->TriangleCount() > 0)
+	{
+		TImplicitMorphology<FDynamicMesh3> ImplicitMorphology;
+		ImplicitMorphology.MorphologyOp = TImplicitMorphology<FDynamicMesh3>::EMorphologyOp::Close;
+		ImplicitMorphology.Source = CachedMorphologyResult.Get();
+		FDynamicMeshAABBTree3 Spatial(ImplicitMorphology.Source, true);
+		ImplicitMorphology.SourceSpatial = &Spatial;
+		ImplicitMorphology.SetCellSizesAndDistance(CachedMorphologyResult->GetBounds(true), NewSettings.CloseDist, NewSettings.VoxelCount, NewSettings.VoxelCount);
+
+		CachedMorphologyResult = MakeShared<FDynamicMesh3, ESPMode::ThreadSafe>(&ImplicitMorphology.Generate());
+		if (CachedMorphologyResult->TriangleCount() == 0)
+		{
+			CachedPostprocessResult = MakeShared<FDynamicMesh3, ESPMode::ThreadSafe>(*InputMesh);
+		}
+	}
+
+
+	if (NewSettings.OpenDist > 0 && CachedMorphologyResult->TriangleCount() > 0)
+	{
+		TImplicitMorphology<FDynamicMesh3> ImplicitMorphology;
+		ImplicitMorphology.MorphologyOp = TImplicitMorphology<FDynamicMesh3>::EMorphologyOp::Open;
+		ImplicitMorphology.Source = CachedMorphologyResult.Get();
+		FDynamicMeshAABBTree3 Spatial(ImplicitMorphology.Source, true);
+		ImplicitMorphology.SourceSpatial = &Spatial;
+		ImplicitMorphology.SetCellSizesAndDistance(CachedMorphologyResult->GetBounds(true), NewSettings.OpenDist, NewSettings.VoxelCount, NewSettings.VoxelCount);
+
+		CachedMorphologyResult = MakeShared<FDynamicMesh3, ESPMode::ThreadSafe>(&ImplicitMorphology.Generate());
+		if (CachedMorphologyResult->TriangleCount() == 0)
+		{
+			CachedPostprocessResult = MakeShared<FDynamicMesh3, ESPMode::ThreadSafe>(*InputMesh);
+		}
+	}
+
+	CachedMorphologySettings = NewSettings;
+	return CachedMorphologyResult;
+}
+
+
+
+
+
+enum class ETargetSpaceType
+{
+	Actor,
+	Component,
+	Frame,
+	Transform,
+	World
+};
+
+
+struct FTargetSpace
+{
+	ETargetSpaceType TargetSpace = ETargetSpaceType::Transform;
+
+	AActor* TargetActor = nullptr;
+	USceneComponent* TargetComponent = nullptr;
+	UE::Geometry::FTransform3d TargetTransform = UE::Geometry::FTransform3d::Identity();
+	FFrame3d TargetFrame = FFrame3d();
+
+	FTargetSpace() {}
+
+	FTargetSpace(AActor* TargetActorIn)
+	{
+		TargetSpace = ETargetSpaceType::Actor;
+		TargetActor = TargetActorIn;
+	}
+
+	FTargetSpace(USceneComponent* TargetComponentIn)
+	{
+		TargetSpace = ETargetSpaceType::Component;
+		TargetComponent = TargetComponentIn;
+	}
+
+	FTargetSpace(const FFrame3d& FrameIn)
+	{
+		TargetSpace = ETargetSpaceType::Frame;
+		TargetFrame = FrameIn;
+	}
+
+	FTargetSpace(const UE::Geometry::FTransform3d& TransformIn)
+	{
+		TargetSpace = ETargetSpaceType::Transform;
+		TargetTransform = TransformIn;
+	}
+
+	static FTargetSpace World()
+	{
+		FTargetSpace Space;
+		Space.TargetSpace = ETargetSpaceType::World;
+		return Space;
+	}
+};
+
+
+
+
+static void TransformToSpace(FDynamicMesh3* Mesh, UPrimitiveComponent* SourceComponent, const FTargetSpace& TargetSpace)
+{
+	// TODO: in certain cases we could combine these transforms, ie if there is no scaling, to avoid multiple MeshTransforms:: calls.
+	// Could also add a variant there that applies multiple transforms in sequence to avoid write overhead
+
+	// transform up to world
+	UE::Geometry::FTransform3d ToWorld(SourceComponent->GetComponentTransform());
+	MeshTransforms::ApplyTransform(*Mesh, ToWorld);
+
+	if (TargetSpace.TargetSpace == ETargetSpaceType::Actor)
+	{
+		check(TargetSpace.TargetActor);
+		UE::Geometry::FTransform3d ActorTransform(TargetSpace.TargetActor->GetActorTransform());
+		UE::Geometry::FTransform3d ToActorTransform = ActorTransform.Inverse();
+		MeshTransforms::ApplyTransform(*Mesh, ToActorTransform);
+
+	}
+	else if (TargetSpace.TargetSpace == ETargetSpaceType::Component)
+	{
+		check(TargetSpace.TargetComponent);
+		UE::Geometry::FTransform3d ComponentTransform(TargetSpace.TargetComponent->GetComponentTransform());
+		UE::Geometry::FTransform3d ToComponentTransform = ComponentTransform.Inverse();
+		MeshTransforms::ApplyTransform(*Mesh, ToComponentTransform);
+	}
+	else if (TargetSpace.TargetSpace == ETargetSpaceType::Frame)
+	{
+		MeshTransforms::WorldToFrameCoords(*Mesh, TargetSpace.TargetFrame);
+	}
+	else if (TargetSpace.TargetSpace == ETargetSpaceType::Transform)
+	{
+		MeshTransforms::ApplyTransform(*Mesh, TargetSpace.TargetTransform);
+	}
+	else if (TargetSpace.TargetSpace == ETargetSpaceType::World)
+	{
+		// we are good
+	}
+	else
+	{
+		check(false);		// how did we get here??
+	}
+}
+
+
+
+
+static bool ExtractStaticMesh(UPrimitiveComponent* Component, int32 LODIndex, const FTargetSpace& TargetSpace, FDynamicMesh3& OutputMesh, bool bWantAttributes = true)
+{
+	if (ensure(Component) == false) return false;
+
+	UStaticMeshComponent* StaticMeshComponent = Cast<UStaticMeshComponent>(Component);
+	if (StaticMeshComponent)
+	{
+		UStaticMesh* StaticMesh = StaticMeshComponent->GetStaticMesh();
+		if (ensure(StaticMesh))
+		{
+			const FMeshDescription* MeshDescription = StaticMesh->GetMeshDescription(LODIndex);
+			if (ensure(MeshDescription))
+			{
+				FMeshDescriptionToDynamicMesh Converter;
+				Converter.bCalculateMaps = false;
+				Converter.bDisableAttributes = (bWantAttributes == false);
+				Converter.bEnableOutputGroups = bWantAttributes;
+
+				OutputMesh.Clear();
+				Converter.Convert(MeshDescription, OutputMesh);
+
+				TransformToSpace(&OutputMesh, StaticMeshComponent, TargetSpace);
+
+				return true;
+			}
+		}
+		return false;		// something failed
+	}
+
+	ensure(false);
+	return false;		// unsupported type
+}
+
+
+
+static bool ExtractStaticMesh(AActor* Actor, int32 LODIndex, const FTargetSpace& TargetSpace, FDynamicMesh3& OutputMesh, bool bWantAttributes = true)
+{
+	if (ensure(Actor) == false) return false;
+
+	AStaticMeshActor* StaticMeshActor = Cast<AStaticMeshActor>(Actor);
+	if (StaticMeshActor)
+	{
+		return ExtractStaticMesh(StaticMeshActor->GetStaticMeshComponent(), LODIndex, TargetSpace, OutputMesh, bWantAttributes);
+	}
+
+	return false;
+}
+
+
+
+
+
+
+class FLoopProcessor
+{
+public:
+	FDynamicMesh3* Mesh = nullptr;
+	FEdgeLoop InitialLoop;
+
+	FEdgeLoop CurrentLoop;
+	TArray<FVector3d> PassBuffer;
+
+	bool bApplySmoothingPerPass = true;
+	double SmoothingAlpha = 0.2;
+
+	FLoopProcessor()
+	{
+	}
+
+	FLoopProcessor(FDynamicMesh3* MeshIn, const FEdgeLoop& LoopIn)
+	{
+		Initialize(MeshIn, LoopIn);
+	}
+
+	void Initialize(FDynamicMesh3* MeshIn, const FEdgeLoop& LoopIn)
+	{
+		Mesh = MeshIn;
+		InitialLoop = LoopIn;
+		CurrentLoop = InitialLoop;
+	}
+
+
+	void ApplyPasses(int32 NumPasses, TFunctionRef<FVector3d(int32, const FVector3d&)> UpdatePositionFunc)
+	{
+		for (int32 k = 0; k < NumPasses; ++k)
+		{
+			ApplyPass(UpdatePositionFunc);
+		}
+	}
+
+
+	void ApplyPass( TFunctionRef<FVector3d(int32, const FVector3d&)> UpdatePositionFunc )
+	{
+		int32 NumV = CurrentLoop.GetVertexCount();
+		PassBuffer.SetNumUninitialized(NumV);
+
+		for (int32 k = 0; k < NumV; ++k)
+		{
+			int32 NextK = (k + 1) % NumV;
+			int32 PrevK = (k == 0) ? (NumV-1) : (k-1);
+
+			int32 vid = CurrentLoop.Vertices[k];
+			FVector3d Pos = Mesh->GetVertex(vid);
+			FVector3d NextPos = Mesh->GetVertex(CurrentLoop.Vertices[NextK]);
+			FVector3d PrevPos = Mesh->GetVertex(CurrentLoop.Vertices[PrevK]);
+
+			if (bApplySmoothingPerPass)
+			{
+				FVector3d Centroid = (NextPos + PrevPos) * 0.5;
+				Pos = UE::Geometry::Lerp(Pos, Centroid, SmoothingAlpha);
+			}
+
+			Pos = UpdatePositionFunc(vid, Pos);
+
+			PassBuffer[k] = Pos;
+		}
+
+		for (int32 k = 0; k < NumV; ++k)
+		{
+			FVector3d NewPos = PassBuffer[k];
+			int32 vid = CurrentLoop.Vertices[k];
+			Mesh->SetVertex(vid, NewPos);
+		}
+	}
+
+
+
+};
+
+
+
+
+TSharedPtr<FDynamicMesh3, ESPMode::ThreadSafe> UGroomToMeshTool::UpdateClipMesh(TSharedPtr<FDynamicMesh3, ESPMode::ThreadSafe> InputMesh)
+{
+	FClipMeshSettings NewSettings;
+	NewSettings.InputMesh = InputMesh;
+	NewSettings.ClipSource = Settings->ClipMeshActor.Get();
+	if (CachedClipMeshSettings == NewSettings)
+	{
+		return CachedClipMeshResult;
+	}
+
+	CachedClipMeshResult = MakeShared<FDynamicMesh3, ESPMode::ThreadSafe>(*InputMesh);
+	FDynamicMesh3* EditMesh = CachedClipMeshResult.Get();
+
+	// do a quick tiny-edges collapse pass
+	double AvgEdgeLen = TMeshQueries<FDynamicMesh3>::AverageEdgeLength(*EditMesh);
+	FQEMSimplification Simplifier(EditMesh);
+	Simplifier.SimplifyToEdgeLength(AvgEdgeLen * 0.75);
+	
+
+	if (NewSettings.ClipSource != nullptr)
+	{
+		FDynamicMesh3 ClipMesh;
+		if (ExtractStaticMesh(NewSettings.ClipSource, 0, FTargetSpace(TargetGroom.Get()), ClipMesh, false))
+		{
+			FDynamicMeshAABBTree3 BVTree(&ClipMesh, true);
+			TFastWindingTree<FDynamicMesh3> WindingTree(&BVTree);
+
+			FMeshVertexSelection SelectedVerts(EditMesh);
+			SelectedVerts.SelectByPosition([&](const FVector3d& Position) { return WindingTree.IsInside(Position); });
+
+			FMeshFaceSelection SelectedFaces(EditMesh, SelectedVerts, 1);
+			//SelectedFaces.ContractBorderByOneRingNeighbours();
+			SelectedFaces.LocalOptimize(true, true);
+
+			FDynamicMeshEditor Editor(EditMesh);
+			Editor.RemoveTriangles(SelectedFaces.AsArray(), true);
+
+			FMeshBoundaryLoops BoundaryLoops(EditMesh, true);
+			for (FEdgeLoop& Loop : BoundaryLoops.Loops)
+			{
+				double ProjectionSpeed = 0.3;
+				double TuneSpeed = 0.9;
+				double SmoothSpeed = 0.25;
+
+				FLoopProcessor Processor(EditMesh, Loop);
+				Processor.bApplySmoothingPerPass = true;
+				Processor.SmoothingAlpha = SmoothSpeed;
+				Processor.ApplyPasses(10, [&](int32 vid, const FVector3d& Pos)
+				{
+					FVector3d NearestPos = BVTree.FindNearestPoint(EditMesh->GetVertex(vid));
+					return UE::Geometry::Lerp(Pos, NearestPos, ProjectionSpeed);
+				});
+
+				Processor.ApplyPasses(10, [&](int32 vid, const FVector3d& Pos)
+				{
+					FVector3d NearestPos = BVTree.FindNearestPoint(EditMesh->GetVertex(vid));
+					return UE::Geometry::Lerp(Pos, NearestPos, TuneSpeed);
+				});
+			}
+
+
+
+		}
+	}
+
+	CachedClipMeshResult->CompactInPlace();
+	CachedClipMeshSettings = NewSettings;
+	return CachedClipMeshResult;
+}
+
+
+
+TSharedPtr<FDynamicMesh3, ESPMode::ThreadSafe> UGroomToMeshTool::UpdateSmoothing(TSharedPtr<FDynamicMesh3, ESPMode::ThreadSafe> InputMesh)
+{
+	FSmoothingSettings NewSettings;
+	NewSettings.InputMesh = InputMesh;
+	NewSettings.Smoothness = Settings->Smoothness;
+	NewSettings.VolumeCorrection = Settings->VolumeCorrection;
+	if (CachedSmoothSettings == NewSettings)
+	{
+		return CachedSmoothResult;
+	}
+
+	CachedSmoothResult = MakeShared<FDynamicMesh3, ESPMode::ThreadSafe>(*InputMesh);
+
+	if (NewSettings.Smoothness > 0)
+	{
+		TSharedPtr<FMeshNormals> VtxNormals = MakeShared<FMeshNormals>(InputMesh.Get());
+		if (NewSettings.VolumeCorrection != 0)
+		{
+			VtxNormals->ComputeVertexNormals();
+		}
+
+		FDynamicMesh3* EditMesh = CachedSmoothResult.Get();
+
+		// normalize mesh area/position
+		// compute area of the input mesh and compute normalization scaling factor
+		FVector2d VolArea = TMeshQueries<FDynamicMesh3>::GetVolumeArea(*EditMesh);
+		double UnitScalingMeasure = FMathd::Max(0.01, FMathd::Sqrt(VolArea.Y / 6.0));  // 6.0 is a bit arbitrary here...surface area of unit box
+		FAxisAlignedBox3d Bounds = InputMesh->GetBounds(true);
+		FVector3d SrcTranslate = Bounds.Center();
+		MeshTransforms::Translate(*EditMesh, -SrcTranslate);
+		double SrcScale = UnitScalingMeasure;
+		MeshTransforms::Scale(*EditMesh, (1.0 / SrcScale) * FVector3d::One(), FVector3d::Zero());
+
+		FSmoothingOpBase::FOptions Options;
+		Options.BaseNormals = VtxNormals;
+		Options.SmoothAlpha = NewSettings.Smoothness;
+		Options.BoundarySmoothAlpha = 0.0;
+		double NonlinearT = FMathd::Pow(NewSettings.Smoothness, 2.0);
+		// this is an empirically-determined hack that seems to work OK to normalize the smoothing result for variable vertex count...
+		double ScaledPower = (NonlinearT / 50.0) * EditMesh->VertexCount();
+		Options.SmoothPower = ScaledPower;
+		Options.bUniform = false;
+		Options.bUseImplicit = true;
+		Options.NormalOffset = NewSettings.VolumeCorrection / 10.0;		// rescale hack
+
+		// weightmap
+		//Options.bUseWeightMap = true;
+		//Options.WeightMap = GetActiveWeightMap();
+		//Options.WeightMapMinMultiplier = WeightMapProperties->MinSmoothMultiplier;
+
+		TUniquePtr<FCotanSmoothingOp> MeshOp = MakeUnique<FCotanSmoothingOp>(EditMesh, Options);
+		MeshOp->CalculateResult(nullptr);
+
+		TUniquePtr<FDynamicMesh3> NewMesh = MeshOp->ExtractResult();
+
+		// invert normalization transform
+		MeshTransforms::Scale(*NewMesh, FVector3d(SrcScale, SrcScale, SrcScale), FVector3d::Zero());
+		MeshTransforms::Translate(*NewMesh, SrcTranslate);
+
+		CachedSmoothResult = TSharedPtr<FDynamicMesh3, ESPMode::ThreadSafe>(NewMesh.Release());
+	}
+
+
+	CachedSmoothSettings = NewSettings;
+	return CachedSmoothResult;
+}
+
+
+
+
+
+TSharedPtr<FDynamicMesh3, ESPMode::ThreadSafe> UGroomToMeshTool::UpdateSimplification(TSharedPtr<FDynamicMesh3, ESPMode::ThreadSafe> InputMesh)
+{
+	FSimplifySettings NewSettings;
+	NewSettings.InputMesh = InputMesh;
+	NewSettings.TargetCount = Settings->VertexCount;
+	if (CachedSimplifySettings == NewSettings)
+	{
+		return CachedSimplifyResult;
+	}
+
+	CachedSimplifyResult = MakeShared<FDynamicMesh3, ESPMode::ThreadSafe>(*InputMesh);
+
+	FQEMSimplification Simplifier(CachedSimplifyResult.Get());
+	Simplifier.SimplifyToVertexCount(NewSettings.TargetCount);
+	CachedSimplifyResult->CompactInPlace();
+
+	CachedSimplifySettings = NewSettings;
+	return CachedSimplifyResult;
+}
+
+
+
+TSharedPtr<FDynamicMesh3, ESPMode::ThreadSafe> UGroomToMeshTool::UpdatePostprocessing(TSharedPtr<FDynamicMesh3, ESPMode::ThreadSafe> InputMesh)
+{
+	FPostprocessSettings NewSettings;
+	NewSettings.InputMesh = InputMesh;
+	NewSettings.UVGenMode = Settings->UVMode;
+	if (CachedPostprocessSettings == NewSettings)
+	{
+		return CachedPostprocessResult;
+	}
+
+	CachedPostprocessResult = MakeShared<FDynamicMesh3, ESPMode::ThreadSafe>(*InputMesh);
+
+	FMeshNormals::InitializeMeshToPerTriangleNormals(CachedPostprocessResult.Get());
+
+	CachedPostprocessResult = UpdateUVs(CachedPostprocessResult, NewSettings.UVGenMode);
+
+	CachedPostprocessSettings = NewSettings;
+	return CachedPostprocessResult;
+}
+
+
+
+
+template<typename EnumerableType>
+void PartitionIndices(EnumerableType Enumerable, 
+	TFunctionRef<FVector3d(int32)> GetPositionForIndexFunc,
+	FLine3d Line, FInterval1d CenterRange,
+	TArray<int32>& BelowOut, TArray<int32>& CenterOut, TArray<int32>& AboveOut)
+{
+	for (int32 vid : Enumerable)
+	{
+		FVector3d Position = GetPositionForIndexFunc(vid);
+		double T = Line.Project(Position);
+		if (CenterRange.Contains(T))
+		{
+			CenterOut.Add(vid);
+		}
+		else if (T < CenterRange.Min)
+		{
+			BelowOut.Add(vid);
+		}
+		else
+		{
+			AboveOut.Add(vid);
+		}
+	}
+}
+
+
+
+TSharedPtr<FDynamicMesh3, ESPMode::ThreadSafe> UGroomToMeshTool::UpdateUVs(TSharedPtr<FDynamicMesh3, ESPMode::ThreadSafe> InputMesh, EGroomToMeshUVMode UVMode)
+{
+	TSharedPtr<FDynamicMesh3, ESPMode::ThreadSafe> Result;
+
+	if (InputMesh->IsClosed())
+	{
+		Result = UpdateUVs_ExpMapPlaneSplits(InputMesh, true);
+		return Result;
+	}
+
+	switch (UVMode)
+	{
+	default:
+	case EGroomToMeshUVMode::PlanarSplitting:
+		Result = UpdateUVs_ExpMapPlaneSplits(InputMesh, false);
+		break;
+	case EGroomToMeshUVMode::MinimalConformal:
+		Result = UpdateUVs_MinimalConformal(InputMesh);
+		break;
+	case EGroomToMeshUVMode::PlanarSplitConformal:
+		Result = UpdateUVs_ExpMapPlaneSplits(InputMesh, true);
+		break;
+	}
+
+	// check that all triangles have valid UVs, because otherwise FUVLayoutOp will crash
+	bool bAllUVsValid = true;
+	auto AttribSet = Result->Attributes()->PrimaryUV();
+	for (int32 TriangleID : Result->TriangleIndicesItr())
+	{
+		bAllUVsValid = bAllUVsValid && AttribSet->IsSetTriangle(TriangleID);
+	}
+	if (!bAllUVsValid)
+	{
+		return Result;
+	}
+
+	// run UV layout
+	FUVLayoutOp UVLayout;
+	UVLayout.OriginalMesh = Result;
+	UVLayout.UVLayoutMode = EUVLayoutOpLayoutModes::RepackToUnitRect;
+	UVLayout.TextureResolution = 64;
+	UVLayout.UVScaleFactor = 1.0;
+	UVLayout.CalculateResult(nullptr);
+	*Result = MoveTemp(*UVLayout.ExtractResult());
+
+	return Result;
+}
+
+
+
+static void GenerateUVTiling(int32 NumTiles, TArray<FAxisAlignedBox2d>& BoxesOut, double UVSpaceGutterWidth)
+{
+	if (NumTiles == 1)
+	{
+		FAxisAlignedBox2d FullBox(FVector2d::Zero(), FVector2d::One());
+		FullBox.Expand(-UVSpaceGutterWidth);
+		BoxesOut.Add(FullBox);
+		return;
+	}
+
+	int32 Columns = (int)FMathd::Ceil(FMathd::Sqrt((double)NumTiles));
+	int32 Rows = Columns;
+
+	double W = 1.0 / (double)Columns;
+
+	for (int32 ri = 0; ri < Rows; ++ri)
+	{
+		double y = (double)ri * W;
+		for (int32 ci = 0; ci < Columns; ++ci)
+		{
+			double x = (double)ci * W;
+
+			FAxisAlignedBox2d Box(FVector2d(x, y), FVector2d(x + W, y + W));
+			Box.Expand(-UVSpaceGutterWidth);
+			BoxesOut.Add(Box);
+			if (BoxesOut.Num() == NumTiles)
+			{
+				return;
+			}
+		}
+	}
+}
+
+
+
+TSharedPtr<FDynamicMesh3, ESPMode::ThreadSafe> UGroomToMeshTool::UpdateUVs_MinimalConformal(TSharedPtr<FDynamicMesh3, ESPMode::ThreadSafe> InputMesh)
+{
+	TSharedPtr<FDynamicMesh3, ESPMode::ThreadSafe> Result = MakeShared<FDynamicMesh3, ESPMode::ThreadSafe>(*InputMesh);
+	FDynamicMesh3* EditMesh = Result.Get();
+
+	FDynamicMeshUVEditor UVEditor(EditMesh, 0, true);
+
+	FMeshConnectedComponents MeshComponents(EditMesh);
+	MeshComponents.FindConnectedTriangles();
+	TArray<FAxisAlignedBox2d> TilingBoxes;
+	GenerateUVTiling(MeshComponents.Num(), TilingBoxes, 0.01);
+	for (int32 k = 0; k < MeshComponents.Num(); ++k)
+	{
+		FMeshConnectedComponents::FComponent& Component = MeshComponents.GetComponent(k);
+		FUVEditResult EditResult;
+		check(Component.Indices.Num() > 1);
+		if ( UVEditor.SetTriangleUVsFromFreeBoundaryConformal(Component.Indices, &EditResult) == false )
+		{
+			// if we failed at conformal, fallback to expmap for this island
+			EditResult = FUVEditResult();
+			if ( UVEditor.SetTriangleUVsFromExpMap(Component.Indices) == false )
+			{
+				// if we somehow failed at conformal, fallback to trivial planar projection
+				UVEditor.SetTriangleUVsFromProjection(Component.Indices, FFrame3d());
+			}
+		}
+
+		// destroys relative scale between islands (but for conformal it is not meaninful anyway...)
+		//UE::MeshUVTransforms::FitToBox(UVEditor.GetOverlay(), EditResult.NewUVElements, TilingBoxes[k], true);
+	}
+
+	return Result;
+}
+
+
+
+
+TSharedPtr<FDynamicMesh3, ESPMode::ThreadSafe> UGroomToMeshTool::UpdateUVs_ExpMapPlaneSplits(TSharedPtr<FDynamicMesh3, ESPMode::ThreadSafe> InputMesh, bool bRecalcAsConformal)
+{
+	int32 SplitAxisIndex1 = 0;
+	int32 SplitAxisIndex2 = 1;
+	double MiddleFraction1 = 0.6;
+	double MiddleFraction2 = 0.6;
+
+	TSharedPtr<FDynamicMesh3, ESPMode::ThreadSafe> Result = MakeShared<FDynamicMesh3, ESPMode::ThreadSafe>(*InputMesh);
+	FDynamicMesh3* EditMesh = Result.Get();
+	
+	FDynamicMeshAABBTree3 BVTree(EditMesh, true);
+
+	FAxisAlignedBox3d Bounds = EditMesh->GetBounds(true);
+	FVector3d Center = Bounds.Center();
+	FVector3d Extents = Bounds.Extents();
+
+	// split left/right/middle
+
+	FVector3d SplitAxis1 = FVector3d::Zero();
+	SplitAxis1[SplitAxisIndex1] = 1.0;
+	FInterval1d AxisInterval1(-Extents[SplitAxisIndex1]*MiddleFraction1, Extents[SplitAxisIndex1]*MiddleFraction1);
+
+	TArray<int32> LeftT, RightT, MiddleT;
+	PartitionIndices(EditMesh->TriangleIndicesItr(), [&](int32 tid) { return EditMesh->GetTriCentroid(tid); },
+		FLine3d(Center, SplitAxis1), AxisInterval1, LeftT, MiddleT, RightT);
+
+	int32 LeftGID = EditMesh->AllocateTriangleGroup();
+	int32 RightGID = EditMesh->AllocateTriangleGroup();
+
+	FMeshFaceSelection MiddleFaces(EditMesh);
+	MiddleFaces.Select(MiddleT);
+	MiddleFaces.LocalOptimize(true, true);
+
+	FMeshFaceSelection LeftFaces(EditMesh);
+	LeftFaces.Select(LeftT);
+	LeftFaces.ExpandToOneRingNeighbours();
+	LeftFaces.ExpandToOneRingNeighbours();
+	LeftFaces.Deselect(MiddleFaces);
+	for (int32 tid : LeftFaces)
+	{
+		EditMesh->SetTriangleGroup(tid, LeftGID);
+	}
+
+	FMeshFaceSelection RightFaces(EditMesh);
+	RightFaces.Select(RightT);
+	RightFaces.ExpandToOneRingNeighbours();
+	RightFaces.ExpandToOneRingNeighbours();
+	RightFaces.Deselect(MiddleFaces);
+	for (int32 tid : RightFaces)
+	{
+		EditMesh->SetTriangleGroup(tid, RightGID);
+	}
+
+
+	// split middle group again
+	
+	FVector3d SplitAxis2 = FVector3d::Zero();
+	SplitAxis2[SplitAxisIndex2] = 1.0;
+	FInterval1d AxisInterval2(-Extents[SplitAxisIndex2] * MiddleFraction2, Extents[SplitAxisIndex2] * MiddleFraction2);
+
+	MiddleT.Reset();
+	TArray<int32> FrontT, BackT;
+	PartitionIndices(MiddleFaces, [&](int32 tid) { return EditMesh->GetTriCentroid(tid); },
+		FLine3d(Center, SplitAxis2), AxisInterval2, FrontT, MiddleT, BackT);
+
+	int32 FrontGID = EditMesh->AllocateTriangleGroup();
+	int32 MiddleGID = EditMesh->AllocateTriangleGroup();
+	int32 BackGID = EditMesh->AllocateTriangleGroup();
+
+	MiddleFaces = FMeshFaceSelection(EditMesh);
+	MiddleFaces.Select(MiddleT);
+	MiddleFaces.LocalOptimize(true, true);
+	for (int32 tid : MiddleFaces)
+	{
+		EditMesh->SetTriangleGroup(tid, MiddleGID);
+	}
+
+	FMeshFaceSelection FrontFaces(EditMesh);
+	FrontFaces.Select(FrontT);
+	FrontFaces.ExpandToOneRingNeighbours();
+	FrontFaces.ExpandToOneRingNeighbours();
+	FrontFaces.Deselect(MiddleFaces);
+	for (int32 tid : FrontFaces)
+	{
+		EditMesh->SetTriangleGroup(tid, FrontGID);
+	}
+
+	FMeshFaceSelection BackFaces(EditMesh);
+	BackFaces.Select(BackT);
+	BackFaces.ExpandToOneRingNeighbours();
+	BackFaces.ExpandToOneRingNeighbours();
+	BackFaces.Deselect(MiddleFaces);
+	for (int32 tid : BackFaces)
+	{
+		EditMesh->SetTriangleGroup(tid, BackGID);
+	}
+
+
+	// generate UVs for all UV islands
+
+	FDynamicMeshUVEditor UVEditor(EditMesh, 0, true);
+
+	FMeshConnectedComponents GroupComponents(EditMesh);
+	GroupComponents.FindConnectedTriangles([EditMesh](int32 Triangle0, int32 Triangle1) {
+		return EditMesh->GetTriangleGroup(Triangle0) == EditMesh->GetTriangleGroup(Triangle1);
+	});
+
+	TArray<FAxisAlignedBox2d> TilingBoxes;
+	GenerateUVTiling(GroupComponents.Num(), TilingBoxes, 0.01);
+	for ( int32 k = 0; k < GroupComponents.Num(); ++k )
+	{
+		FMeshConnectedComponents::FComponent& Component = GroupComponents.GetComponent(k);
+		if (bRecalcAsConformal && Component.Indices.Num() > 1)
+		{
+			if (UVEditor.SetTriangleUVsFromFreeBoundaryConformal(Component.Indices) == false)
+			{
+				// if we failed at conformal, fallback to expmap for this island
+				if (UVEditor.SetTriangleUVsFromExpMap(Component.Indices) == false)
+				{
+					// if we somehow failed at conformal, fallback to trivial planar projection
+					UVEditor.SetTriangleUVsFromProjection(Component.Indices, FFrame3d());
+				}
+			}
+		}
+		else
+		{
+			UVEditor.SetTriangleUVsFromExpMap(Component.Indices);
+		}
+	}
+
+	return Result;
+}
+
+
+
+void UGroomToMeshTool::UpdateLineSet()
+{
+	FString GuideCurvesName = TEXT("HAIRSTRANDS_CURVES");
+	if (Settings->bShowGuides)
+	{
+		check(TargetGroom.IsValid());
+		AGroomActor* GroomActor = TargetGroom.Get();
+		PreviewGeom->CreateOrUpdateLineSet(GuideCurvesName, 1, [&](int32 CurveIndex, TArray<FRenderableLine>& LinesOut)
+		{
+			ProcessHairCurves(GroomActor, true, [&](const TArrayView<FVector3f>& Positions, const TArrayView<float>& Radii)
+			{
+				int32 Count = Positions.Num() - 1;
+				for (int32 k = 0; k < Count; ++k)
+				{
+					LinesOut.Add(FRenderableLine(Positions[k], Positions[k + 1], FColor::Red, 0.1f));
+				}
+			});
+		});
+	}
+	else
+	{
+		PreviewGeom->RemoveLineSet(GuideCurvesName);
+	}
+}
+
+
+
+
+void UGroomToMeshTool::UpdatePreview(TSharedPtr<FDynamicMesh3, ESPMode::ThreadSafe> ResultMesh)
+{
+	CurrentMesh = FDynamicMesh3(*ResultMesh);
+
+	// compute normals
+	if (CurrentMesh.TriangleCount() == 0)
+	{
+		GetToolManager()->DisplayMessage(LOCTEXT("NoTrianglesWarning", "The Output Mesh does not contain any triangles with the current settings"), EToolMessageLevel::UserWarning);
+		PreviewMesh->UpdatePreview(&CurrentMesh);
+	}
+	else
+	{
+		GetToolManager()->DisplayMessage({}, EToolMessageLevel::UserWarning);
+
+		FMeshNormals::InitializeOverlayToPerVertexNormals(CurrentMesh.Attributes()->PrimaryNormals(), false);
+
+		if (Settings->bShowSideBySide)
+		{
+			FAxisAlignedBox3d Bounds = CurrentMesh.GetBounds();
+			FDynamicMesh3 TmpMesh(CurrentMesh);
+			MeshTransforms::Translate(TmpMesh, Bounds.Width() * FVector3d::UnitX());
+			PreviewMesh->UpdatePreview(&TmpMesh);
+		}
+		else
+		{
+			PreviewMesh->UpdatePreview(&CurrentMesh);
+		}
+	}
+
+	UpdateLineSet();
+}
+
+
+
+#undef LOCTEXT_NAMESPACE
diff --git a/Engine/Plugins/Experimental/HairModelingToolset/Source/HairModelingToolset/Private/HairModelingToolCommands.cpp b/Engine/Plugins/Experimental/HairModelingToolset/Source/HairModelingToolset/Private/HairModelingToolCommands.cpp
new file mode 100644
--- /dev/null
+++ b/Engine/Plugins/Experimental/HairModelingToolset/Source/HairModelingToolset/Private/HairModelingToolCommands.cpp
@@ -0,0 +1,33 @@
+// Copyright Epic Games, Inc. All Rights Reserved.
+
+#include "HairModelingToolCommands.h"
+#include "EditorStyleSet.h"
+#include "InputCoreTypes.h"
+#include "HairModelingToolsStyle.h"
+
+#define LOCTEXT_NAMESPACE "HairModelingToolCommands"
+
+
+
+FHairModelingToolCommands::FHairModelingToolCommands() :
+	TCommands<FHairModelingToolCommands>(
+		"HairModelingToolCommands", // Context name for fast lookup
+		NSLOCTEXT("Contexts", "HairModelingToolCommands", "Hair Modeling Tools"), // Localized context name for displaying
+		NAME_None, // Parent
+		FHairModelingToolsStyle::Get()->GetStyleSetName() // Icon Style Set
+		)
+{
+}
+
+
+void FHairModelingToolCommands::RegisterCommands()
+{
+	UI_COMMAND(BeginGroomToMeshTool, "Helmet", "Generate Helmet Mesh for Selected Groom", EUserInterfaceActionType::ToggleButton, FInputChord());
+	UI_COMMAND(BeginGenerateLODMeshesTool, "HlmLOD", "Generate LODS for Hair Helmet", EUserInterfaceActionType::ToggleButton, FInputChord());
+	UI_COMMAND(BeginGroomCardsEditorTool, "CardsEd", "Edit Hair Cards", EUserInterfaceActionType::ToggleButton, FInputChord());
+}
+
+
+
+
+#undef LOCTEXT_NAMESPACE
diff --git a/Engine/Plugins/Experimental/HairModelingToolset/Source/HairModelingToolset/Private/HairModelingToolsStyle.cpp b/Engine/Plugins/Experimental/HairModelingToolset/Source/HairModelingToolset/Private/HairModelingToolsStyle.cpp
new file mode 100644
--- /dev/null
+++ b/Engine/Plugins/Experimental/HairModelingToolset/Source/HairModelingToolset/Private/HairModelingToolsStyle.cpp
@@ -0,0 +1,77 @@
+// Copyright Epic Games, Inc. All Rights Reserved.
+
+#include "HairModelingToolsStyle.h"
+#include "Styling/SlateStyleRegistry.h"
+#include "Styling/SlateTypes.h"
+#include "Styling/CoreStyle.h"
+#include "EditorStyleSet.h"
+#include "Interfaces/IPluginManager.h"
+#include "SlateOptMacros.h"
+#include "Styling/SlateStyleMacros.h"
+
+#define IMAGE_PLUGIN_BRUSH( RelativePath, ... ) FSlateImageBrush( FHairModelingToolsStyle::InContent( RelativePath, ".png" ), __VA_ARGS__ )
+
+// This is to fix the issue that SlateStyleMacros like IMAGE_BRUSH look for RootToContentDir but StyleSet->RootToContentDir is how this style is set up
+#define RootToContentDir StyleSet->RootToContentDir
+
+FString FHairModelingToolsStyle::InContent(const FString& RelativePath, const ANSICHAR* Extension)
+{
+	static FString ContentDir = IPluginManager::Get().FindPlugin(TEXT("HairModelingTools"))->GetContentDir();
+	return (ContentDir / RelativePath) + Extension;
+}
+
+TSharedPtr< FSlateStyleSet > FHairModelingToolsStyle::StyleSet = nullptr;
+TSharedPtr< class ISlateStyle > FHairModelingToolsStyle::Get() { return StyleSet; }
+
+FName FHairModelingToolsStyle::GetStyleSetName()
+{
+	static FName ModelingToolsStyleName(TEXT("HairModelingToolsStyle"));
+	return ModelingToolsStyleName;
+}
+
+BEGIN_SLATE_FUNCTION_BUILD_OPTIMIZATION
+
+void FHairModelingToolsStyle::Initialize()
+{
+	// Const icon sizes
+	const FVector2D Icon20x20(20.0f, 20.0f);
+
+	// Only register once
+	if (StyleSet.IsValid())
+	{
+		return;
+	}
+
+	StyleSet = MakeShareable(new FSlateStyleSet(GetStyleSetName()));
+	StyleSet->SetContentRoot(IPluginManager::Get().FindPlugin(TEXT("HairModelingToolset"))->GetContentDir());
+	StyleSet->SetCoreContentRoot(FPaths::EngineContentDir() / TEXT("Slate"));
+
+	{
+		StyleSet->Set("HairModelingToolCommands.BeginGroomCardsEditorTool", new IMAGE_BRUSH_SVG("Icons/CardsEditor", Icon20x20));
+		StyleSet->Set("HairModelingToolCommands.BeginGroomCardsEditorTool.Small", new IMAGE_BRUSH_SVG("Icons/CardsEditor", Icon20x20));
+		StyleSet->Set("HairModelingToolCommands.BeginGenerateLODMeshesTool", new IMAGE_BRUSH_SVG("Icons/GenLODs", Icon20x20));
+		StyleSet->Set("HairModelingToolCommands.BeginGenerateLODMeshesTool.Small", new IMAGE_BRUSH_SVG("Icons/GenLODs", Icon20x20));
+		StyleSet->Set("HairModelingToolCommands.BeginGroomToMeshTool", new IMAGE_BRUSH_SVG("Icons/HairHelmet", Icon20x20));
+		StyleSet->Set("HairModelingToolCommands.BeginGroomToMeshTool.Small", new IMAGE_BRUSH_SVG("Icons/HairHelmet", Icon20x20));
+
+	}
+
+	FSlateStyleRegistry::RegisterSlateStyle(*StyleSet.Get());
+};
+
+END_SLATE_FUNCTION_BUILD_OPTIMIZATION
+
+#undef IMAGE_PLUGIN_BRUSH
+#undef IMAGE_BRUSH
+#undef BOX_BRUSH
+#undef DEFAULT_FONT
+
+void FHairModelingToolsStyle::Shutdown()
+{
+	if (StyleSet.IsValid())
+	{
+		FSlateStyleRegistry::UnRegisterSlateStyle(*StyleSet.Get());
+		ensure(StyleSet.IsUnique());
+		StyleSet.Reset();
+	}
+}
diff --git a/Engine/Plugins/Experimental/HairModelingToolset/Source/HairModelingToolset/Private/HairModelingToolsetModule.cpp b/Engine/Plugins/Experimental/HairModelingToolset/Source/HairModelingToolset/Private/HairModelingToolsetModule.cpp
new file mode 100644
--- /dev/null
+++ b/Engine/Plugins/Experimental/HairModelingToolset/Source/HairModelingToolset/Private/HairModelingToolsetModule.cpp
@@ -0,0 +1,68 @@
+// Copyright Epic Games, Inc. All Rights Reserved.
+
+#include "HairModelingToolsetModule.h"
+#include "HairModelingToolsStyle.h"
+#include "HairModelingToolCommands.h"
+
+#include "GroomToMeshTool.h"
+#include "GroomCardsEditorTool.h"
+#include "GenerateLODMeshesTool.h"
+
+#define LOCTEXT_NAMESPACE "FHairModelingToolsetModule"
+
+
+void FHairModelingToolsetModule::StartupModule()
+{
+	FHairModelingToolsStyle::Initialize();
+	FHairModelingToolCommands::Register();
+
+	IModularFeatures::Get().RegisterModularFeature(IModelingModeToolExtension::GetModularFeatureName(), this);
+}
+
+void FHairModelingToolsetModule::ShutdownModule()
+{
+	IModularFeatures::Get().UnregisterModularFeature(IModelingModeToolExtension::GetModularFeatureName(), this);
+
+	FHairModelingToolCommands::Unregister();
+	FHairModelingToolsStyle::Shutdown();
+}
+
+
+FText FHairModelingToolsetModule::GetExtensionName()
+{
+	return LOCTEXT("ExtensionName", "HairTools");
+}
+
+FText FHairModelingToolsetModule::GetToolSectionName()
+{
+	return LOCTEXT("SectionName", "Hair");
+}
+
+void FHairModelingToolsetModule::GetExtensionTools(const FExtensionToolQueryInfo& QueryInfo, TArray<FExtensionToolDescription>& ToolsOut)
+{
+	FExtensionToolDescription GroomToMeshToolInfo;
+	GroomToMeshToolInfo.ToolName = LOCTEXT("HairGroomToMeshTool", "Groom To Mesh");
+	GroomToMeshToolInfo.ToolCommand = FHairModelingToolCommands::Get().BeginGroomToMeshTool;
+	GroomToMeshToolInfo.ToolBuilder = NewObject<UGroomToMeshToolBuilder>();
+	ToolsOut.Add(GroomToMeshToolInfo);
+
+	FExtensionToolDescription GenerateHairLODMeshesToolInfo;
+	GenerateHairLODMeshesToolInfo.ToolName = LOCTEXT("HairLODMeshesTool", "Hair LODs");
+	GenerateHairLODMeshesToolInfo.ToolCommand = FHairModelingToolCommands::Get().BeginGenerateLODMeshesTool;
+	GenerateHairLODMeshesToolInfo.ToolBuilder = NewObject<UGenerateLODMeshesToolBuilder>();
+	ToolsOut.Add(GenerateHairLODMeshesToolInfo);
+
+	FExtensionToolDescription GroomCardsEditorToolInfo;
+	GroomCardsEditorToolInfo.ToolName = LOCTEXT("GroomCardsEditorTool", "Cards Editor");
+	GroomCardsEditorToolInfo.ToolCommand = FHairModelingToolCommands::Get().BeginGroomCardsEditorTool;
+	GroomCardsEditorToolInfo.ToolBuilder = NewObject<UGroomCardsEditorToolBuilder>();
+	ToolsOut.Add(GroomCardsEditorToolInfo);
+
+}
+
+
+
+
+#undef LOCTEXT_NAMESPACE
+	
+IMPLEMENT_MODULE(FHairModelingToolsetModule, HairModelingToolset)
diff --git a/Engine/Plugins/Experimental/HairModelingToolset/Source/HairModelingToolset/Public/GenerateLODMeshesTool.h b/Engine/Plugins/Experimental/HairModelingToolset/Source/HairModelingToolset/Public/GenerateLODMeshesTool.h
new file mode 100644
--- /dev/null
+++ b/Engine/Plugins/Experimental/HairModelingToolset/Source/HairModelingToolset/Public/GenerateLODMeshesTool.h
@@ -0,0 +1,192 @@
+// Copyright Epic Games, Inc. All Rights Reserved.
+
+#pragma once
+
+#include "CoreMinimal.h"
+#include "UObject/NoExportTypes.h"
+#include "SingleSelectionTool.h"
+#include "InteractiveToolBuilder.h"
+#include "DynamicMesh/DynamicMesh3.h"
+#include "DynamicMesh/DynamicMeshAABBTree3.h"
+#include "MeshOpPreviewHelpers.h"
+#include "CleaningOps/SimplifyMeshOp.h"
+#include "Properties/MeshStatisticsProperties.h"
+#include "Properties/RemeshProperties.h"
+#include "BaseTools/SingleSelectionMeshEditingTool.h"
+#include "GenerateLODMeshesTool.generated.h"
+
+/**
+ *
+ */
+UCLASS()
+class HAIRMODELINGTOOLSET_API UGenerateLODMeshesToolBuilder : public USingleSelectionMeshEditingToolBuilder
+{
+	GENERATED_BODY()
+public:
+	virtual USingleSelectionMeshEditingTool* CreateNewTool(const FToolBuilderState& SceneState) const override;
+};
+
+
+
+USTRUCT()
+struct HAIRMODELINGTOOLSET_API FLODLevelGenerateSettings
+{
+	GENERATED_BODY()
+
+	/** Simplification Scheme  */
+	UPROPERTY(EditAnywhere, Category = Options)
+	ESimplifyType SimplifierType = ESimplifyType::UEStandard;
+
+	/** Simplification Target Type  */
+	UPROPERTY(EditAnywhere, Category = Options)
+	ESimplifyTargetType TargetMode = ESimplifyTargetType::Percentage;
+
+	/** Target percentage */
+	UPROPERTY(EditAnywhere, Category = Options, meta = (UIMin = "0", UIMax = "100", EditConditionHides, EditCondition = "TargetMode == ESimplifyTargetType::Percentage"))
+	int32 TargetPercentage = 100;
+
+	/** Target vertex/triangle count */
+	UPROPERTY(EditAnywhere, Category = Options, meta = (UIMin = "4", UIMax = "10000", ClampMin = "1", ClampMax = "9999999999", EditConditionHides, EditCondition = "TargetMode == ESimplifyTargetType::TriangleCount || TargetMode == ESimplifyTargetType::VertexCount"))
+	int32 TargetCount = 500;
+
+	/** Target vertex/triangle count */
+	UPROPERTY(EditAnywhere, AdvancedDisplay, Category = Options)
+	bool bReproject = false;
+
+	UPROPERTY(VisibleAnywhere, Category = Options)
+	FString Result;
+
+	bool operator!=(const FLODLevelGenerateSettings& Other)
+	{
+		return SimplifierType != Other.SimplifierType || TargetMode != Other.TargetMode || TargetPercentage != Other.TargetPercentage || TargetCount != Other.TargetCount || bReproject != Other.bReproject;
+	}
+};
+
+
+
+/**
+ * Standard properties of the Simplify operation
+ */
+UCLASS()
+class HAIRMODELINGTOOLSET_API UGenerateLODMeshesToolProperties : public UMeshConstraintProperties
+{
+	GENERATED_BODY()
+public:
+	UGenerateLODMeshesToolProperties();
+
+	/** Simplification Target Type  */
+	//UPROPERTY(EditAnywhere, Category = Options)
+	UPROPERTY()
+	ESimplifyTargetType TargetMode = ESimplifyTargetType::Percentage;
+
+	/** Simplification Scheme  */
+	//UPROPERTY(EditAnywhere, Category = Options)
+	UPROPERTY()
+	ESimplifyType SimplifierType = ESimplifyType::UEStandard;
+
+	/** Output LOD Assets will be numbered starting at this number */
+	UPROPERTY(EditAnywhere, Category = Options)
+	int NameIndexBase = 0;
+
+	/** Target percentage of original triangle count */
+	//UPROPERTY(EditAnywhere, Category = Options, meta = (UIMin = "0", UIMax = "100", EditCondition = "TargetMode == ESimplifyTargetType::Percentage"))
+	UPROPERTY()
+	int TargetPercentage = 50;
+
+	/** Target edge length */
+	//UPROPERTY(EditAnywhere, Category = Options, meta = (UIMin = "3.0", UIMax = "10.0", ClampMin = "0.001", ClampMax = "1000.0", EditCondition = "TargetMode == ESimplifyTargetType::EdgeLength && SimplifierType != ESimplifyType::UEStandard"))
+	UPROPERTY()
+	float TargetEdgeLength;
+
+	/** Target triangle/vertex count */
+	//UPROPERTY(EditAnywhere, Category = Options, meta = (UIMin = "4", UIMax = "10000", ClampMin = "1", ClampMax = "9999999999", EditCondition = "TargetMode == ESimplifyTargetType::TriangleCount"))
+	UPROPERTY()
+	int TargetCount = 1000;
+
+	/** If true, UVs and Normals are discarded  */
+	//UPROPERTY(EditAnywhere, Category = Options)
+	UPROPERTY()
+	bool bDiscardAttributes = false;
+
+	/** If true, display wireframe */
+	UPROPERTY(EditAnywhere, Category = Display)
+	bool bShowWireframe = false;
+
+	/** Display colors corresponding to the mesh's polygon groups */
+	//UPROPERTY(EditAnywhere, Category = Display)
+	UPROPERTY()
+	bool bShowGroupColors = false;
+
+	/** Enable projection back to input mesh */
+	//UPROPERTY(EditAnywhere, Category = Options, AdvancedDisplay)
+	UPROPERTY()
+	bool bReproject;
+
+	UPROPERTY(EditAnywhere, Category = Options)
+	TArray<FLODLevelGenerateSettings> LODLevels;
+
+};
+
+
+
+class FGenerateLODOperatorFactory : public UE::Geometry::IDynamicMeshOperatorFactory
+{
+public:
+	TWeakObjectPtr<UGenerateLODMeshesTool> ParentTool;
+	FLODLevelGenerateSettings LODSettings;
+	FTransform UseTransform;
+
+	// IDynamicMeshOperatorFactory API
+	virtual TUniquePtr<UE::Geometry::FDynamicMeshOperator> MakeNewOperator() override;
+};
+
+
+
+/**
+ * Simple Mesh Simplifying Tool
+ */
+UCLASS()
+class HAIRMODELINGTOOLSET_API UGenerateLODMeshesTool : public USingleSelectionMeshEditingTool
+{
+	GENERATED_BODY()
+
+public:
+	virtual void Setup() override;
+	virtual void Shutdown(EToolShutdownType ShutdownType) override;
+
+	virtual void OnTick(float DeltaTime) override;
+	virtual void Render(IToolsContextRenderAPI* RenderAPI) override;
+
+	virtual bool HasCancel() const override { return true; }
+	virtual bool HasAccept() const override { return true; }
+	virtual bool CanAccept() const override;
+
+
+
+private:
+	UPROPERTY()
+	TObjectPtr<UGenerateLODMeshesToolProperties> SimplifyProperties;
+
+	UPROPERTY()
+	TArray<TObjectPtr<UMeshOpPreviewWithBackgroundCompute>> Previews;
+
+	TArray<TUniquePtr<FGenerateLODOperatorFactory>> PreviewFactories;
+	friend class FGenerateLODOperatorFactory;
+
+	void UpdateNumPreviews();
+	void InvalidateAllPreviews();
+
+	TSharedPtr<FMeshDescription, ESPMode::ThreadSafe> OriginalMeshDescription;
+	// Dynamic Mesh versions precomputed in Setup (rather than recomputed for every simplify op)
+	TSharedPtr<UE::Geometry::FDynamicMesh3, ESPMode::ThreadSafe> OriginalMesh;
+	TSharedPtr<UE::Geometry::FDynamicMeshAABBTree3, ESPMode::ThreadSafe> OriginalMeshSpatial;
+
+	TArray<FLODLevelGenerateSettings> CachedLODLevels;
+
+	UE::Geometry::FAxisAlignedBox3d WorldBounds;
+
+	void GenerateAssets();
+	void UpdateVisualization();
+
+	void OnPreviewUpdated(UMeshOpPreviewWithBackgroundCompute*);
+};
diff --git a/Engine/Plugins/Experimental/HairModelingToolset/Source/HairModelingToolset/Public/GroomCardsEditorTool.h b/Engine/Plugins/Experimental/HairModelingToolset/Source/HairModelingToolset/Public/GroomCardsEditorTool.h
new file mode 100644
--- /dev/null
+++ b/Engine/Plugins/Experimental/HairModelingToolset/Source/HairModelingToolset/Public/GroomCardsEditorTool.h
@@ -0,0 +1,324 @@
+// Copyright Epic Games, Inc. All Rights Reserved.
+
+#pragma once
+
+#include "CoreMinimal.h"
+#include "Templates/PimplPtr.h"
+#include "BaseTools/MeshSurfacePointTool.h"
+#include "InteractiveToolBuilder.h"
+#include "PreviewMesh.h"
+#include "Drawing/PreviewGeometryActor.h"
+#include "Selection/PolygonSelectionMechanic.h"
+#include "Mechanics/SpaceCurveDeformationMechanic.h"
+#include "Transforms/MultiTransformer.h"
+
+#include "DynamicMesh/DynamicMesh3.h"
+#include "DynamicMesh/DynamicMeshAABBTree3.h"
+#include "GroupTopology.h"
+
+#include "GroomCardsEditorTool.generated.h"
+
+class AGroomActor;
+class AStaticMeshActor;
+
+class FEditableGroomCardSet;
+class FGroomCardEdit;
+class FMeshVertexChangeBuilder;
+
+/**
+ *
+ */
+UCLASS()
+class HAIRMODELINGTOOLSET_API UGroomCardsEditorToolBuilder : public UMeshSurfacePointToolBuilder
+{
+	GENERATED_BODY()
+public:
+	virtual UMeshSurfacePointTool* CreateNewTool(const FToolBuilderState& SceneState) const override;
+};
+
+
+
+
+UENUM()
+enum class EEditGroomCardsToolActions
+{
+	NoAction,
+	Delete,
+	SelectionClear,
+	SelectionFill,
+	SelectionAddNext,
+	SelectionAddPrevious,
+	SelectionAddToEnd,
+	SelectionAddToStart
+};
+
+
+UCLASS()
+class HAIRMODELINGTOOLSET_API UEditGroomCardsToolActionPropertySet : public UInteractiveToolPropertySet
+{
+	GENERATED_BODY()
+
+public:
+	TWeakObjectPtr<UGroomCardsEditorTool> ParentTool;
+
+	void Initialize(UGroomCardsEditorTool* ParentToolIn) { ParentTool = ParentToolIn; }
+
+	void PostAction(EEditGroomCardsToolActions Action);
+};
+
+
+UCLASS()
+class HAIRMODELINGTOOLSET_API USelectGroomCardsToolActions : public UEditGroomCardsToolActionPropertySet
+{
+	GENERATED_BODY()
+public:
+	/** Clear the current selection */
+	UFUNCTION(CallInEditor, Category = Selection, meta = (DisplayPriority = 1))
+	void Clear() { PostAction(EEditGroomCardsToolActions::SelectionClear); }
+
+	/** Select unselected points along curve between selected points */
+	UFUNCTION(CallInEditor, Category = Selection, meta = (DisplayPriority = 1))
+	void Fill() { PostAction(EEditGroomCardsToolActions::SelectionFill); }
+
+	/** Add the next vertex along the card curve to the selection */
+	UFUNCTION(CallInEditor, Category = Selection, meta = (DisplayPriority = 3))
+	void AddNext() { PostAction(EEditGroomCardsToolActions::SelectionAddNext); }
+
+	/** Add the previous vertex along the card curve to the selection */
+	UFUNCTION(CallInEditor, Category = Selection, meta = (DisplayPriority = 4))
+	void AddPrevious() { PostAction(EEditGroomCardsToolActions::SelectionAddPrevious); }
+
+	/** Select all vertices to the end of the card */
+	UFUNCTION(CallInEditor, Category = Selection, meta = (DisplayPriority = 5))
+	void ToEnd() { PostAction(EEditGroomCardsToolActions::SelectionAddToEnd); }
+
+	/** Select all vertices to the start of the card */
+	UFUNCTION(CallInEditor, Category = Selection, meta = (DisplayPriority = 6))
+	void ToStart() { PostAction(EEditGroomCardsToolActions::SelectionAddToStart); }
+
+};
+
+
+
+UCLASS()
+class HAIRMODELINGTOOLSET_API UEditGroomCardsToolActions : public UEditGroomCardsToolActionPropertySet
+{
+	GENERATED_BODY()
+public:
+	/** Delete the current selected cards */
+	UFUNCTION(CallInEditor, Category = FaceEdits, meta = (DisplayName = "Delete", DisplayPriority = 10))
+	void Delete() { PostAction(EEditGroomCardsToolActions::Delete); }
+};
+
+
+
+
+
+
+
+
+
+UCLASS()
+class HAIRMODELINGTOOLSET_API UGroomCardsInfoToolProperties : public UInteractiveToolPropertySet
+{
+	GENERATED_BODY()
+
+public:
+	UPROPERTY(VisibleAnywhere, Category = Statistics)
+	int32 NumCards = 0;
+
+	UPROPERTY(VisibleAnywhere, Category = Statistics)
+	int32 NumVertices = 0;
+
+	UPROPERTY(VisibleAnywhere, Category = Statistics)
+	int32 NumTriangles = 0;
+
+};
+
+
+
+/**
+ *
+ */
+UCLASS()
+class HAIRMODELINGTOOLSET_API UGroomCardsEditorTool : public UMeshSurfacePointTool, public IClickBehaviorTarget
+{
+	GENERATED_BODY()
+
+public:
+	UGroomCardsEditorTool();
+
+	virtual void SetWorld(UWorld* World) { this->TargetWorld = World; }
+
+	virtual void Setup() override;
+	virtual void Shutdown(EToolShutdownType ShutdownType) override;
+
+	virtual void Render(IToolsContextRenderAPI* RenderAPI) override;
+	virtual void OnTick(float DeltaTime) override;
+
+	virtual bool HasCancel() const override { return true; }
+	virtual bool HasAccept() const override;
+	virtual bool CanAccept() const override;
+
+
+	// UMeshSurfacePointTool API
+	virtual bool HitTest(const FRay& Ray, FHitResult& OutHit) override;
+	virtual void OnBeginDrag(const FRay& Ray) override;
+	virtual void OnUpdateDrag(const FRay& Ray) override;
+	virtual void OnEndDrag(const FRay& Ray) override;
+	virtual bool OnUpdateHover(const FInputDeviceRay& DevicePos) override;
+	virtual void OnEndHover() override;
+
+	// IClickDragBehaviorTarget API
+	virtual FInputRayHit CanBeginClickDragSequence(const FInputDeviceRay& PressPos) override;
+
+	// IClickBehaviorTarget API
+	virtual FInputRayHit IsHitByClick(const FInputDeviceRay& ClickPos) override;
+	virtual void OnClicked(const FInputDeviceRay& ClickPos) override;
+
+
+	// actions API
+	void RequestAction(EEditGroomCardsToolActions ActionType);
+
+protected:
+	UPROPERTY()
+	TObjectPtr<USelectGroomCardsToolActions> SelectActions;
+
+	UPROPERTY()
+	TObjectPtr<UEditGroomCardsToolActions> EditActions;
+
+	UPROPERTY()
+	TObjectPtr<UGroomCardsInfoToolProperties> InfoProperties;
+
+	UPROPERTY()
+	TObjectPtr<UPreviewMesh> PreviewMesh;
+
+	UPROPERTY()
+	TLazyObjectPtr<AGroomActor> TargetGroom;
+
+	UPROPERTY()
+	TObjectPtr<UPreviewGeometry> PreviewGeom;
+
+protected:		// materials
+
+	UPROPERTY()
+	TObjectPtr<UMaterialInterface> MeshMaterial;
+
+	UPROPERTY()
+	TObjectPtr<UMaterialInterface> UVMaterial;
+
+protected:		// mechanics
+
+	UPROPERTY()
+	TObjectPtr<UPolygonSelectionMechanic> CardMeshSelectionMechanic;
+
+	UPROPERTY()
+	TObjectPtr<USpaceCurveDeformationMechanic> ControlPointsMechanic;
+
+
+	bool bSelectionStateDirty = false;
+	void OnSelectionModifiedEvent();
+
+
+protected:
+	UWorld* TargetWorld = nullptr;
+
+	TPimplPtr<FEditableGroomCardSet> EditableCardSet;
+
+	void BeginCardEdit(int32 CardGroupID);
+	void RestoreCardEdit(const FGroomCardEdit* RestoreEdit);
+	void EndCardEdit();
+	void CompleteActiveCardEditAndEmitChanges();
+
+	TPimplPtr<FGroomCardEdit> ActiveCardEdit;
+	int32 ActiveCardGroupID = -1;
+	bool bCurveUpdatePending = false;
+	bool bActiveCardEditUpdated = false;
+	void UpdateOnCurveEdit();
+
+	FMeshVertexChangeBuilder* ActiveVertexChange = nullptr;
+	void BeginMoveChange();
+	void EndMoveChange();
+
+	TSharedPtr<FSpaceCurveSource> CurveSourceAdapter;
+
+
+	void UpdateLineSet();
+
+	bool bVisualizationChanged = false;
+
+	bool bSetupValid = false;
+	void InitializeMesh();
+
+	TUniquePtr<FGroupTopology> Topology;
+	void RecomputeTopology();
+
+	void OnPreviewMeshChanged();
+	FDelegateHandle OnPreviewMeshChangedHandle;
+
+protected:
+	// operation support
+
+	EEditGroomCardsToolActions PendingAction = EEditGroomCardsToolActions::NoAction;
+
+	int32 CurrentOperationTimestamp = 1;
+	bool CheckInOperation(int32 Timestamp) const { return CurrentOperationTimestamp == Timestamp; }
+
+	void CompleteMeshEditChange(const FText& TransactionLabel, TUniquePtr<FToolCommandChange> EditChange, const FGroupTopologySelection& OutputSelection);
+
+	void AfterTopologyEdit();
+	int32 ModifiedTopologyCounter = 0;
+	bool bWasTopologyEdited = false;
+
+	void ApplyDeleteAction();
+
+
+	friend class FEditGroomCardsTopologyPreEditChange;
+	friend class FEditGroomCardsTopologyPostEditChange;
+	friend class FBeginGroomCardsEditChange;
+	friend class FEndGroomCardsEditChange;
+
+public:
+	void ApplyGroomCardEdit(const FGroomCardEdit& Edit, bool bIsRevert);
+};
+
+
+
+class HAIRMODELINGTOOLSET_API FEditGroomCardsTopologyPreEditChange : public FToolCommandChange
+{
+public:
+	virtual void Apply(UObject* Object) override;
+	virtual void Revert(UObject* Object) override;
+	virtual FString ToString() const override;
+};
+
+class HAIRMODELINGTOOLSET_API FEditGroomCardsTopologyPostEditChange : public FToolCommandChange
+{
+public:
+	virtual void Apply(UObject* Object) override;
+	virtual void Revert(UObject* Object) override;
+	virtual FString ToString() const override;
+};
+
+
+class HAIRMODELINGTOOLSET_API FBeginGroomCardsEditChange : public FToolCommandChange
+{
+public:
+	int32 CardGroupID = -1;
+	FBeginGroomCardsEditChange(int32 GroupID = 0) { CardGroupID = GroupID; }
+
+	virtual void Apply(UObject* Object) override;
+	virtual void Revert(UObject* Object) override;
+	virtual FString ToString() const override;
+};
+
+
+class HAIRMODELINGTOOLSET_API FEndGroomCardsEditChange : public FToolCommandChange
+{
+public:
+	TPimplPtr<FGroomCardEdit> CardEdit;
+
+	virtual void Apply(UObject* Object) override;
+	virtual void Revert(UObject* Object) override;
+	virtual FString ToString() const override;
+};
\ No newline at end of file
diff --git a/Engine/Plugins/Experimental/HairModelingToolset/Source/HairModelingToolset/Public/GroomToMeshTool.h b/Engine/Plugins/Experimental/HairModelingToolset/Source/HairModelingToolset/Public/GroomToMeshTool.h
new file mode 100644
--- /dev/null
+++ b/Engine/Plugins/Experimental/HairModelingToolset/Source/HairModelingToolset/Public/GroomToMeshTool.h
@@ -0,0 +1,271 @@
+// Copyright Epic Games, Inc. All Rights Reserved.
+
+#pragma once
+
+#include "CoreMinimal.h"
+#include "SingleSelectionTool.h"
+#include "InteractiveToolBuilder.h"
+#include "DynamicMesh/DynamicMesh3.h"
+#include "PreviewMesh.h"
+#include "Drawing/PreviewGeometryActor.h"
+
+#include "GroomToMeshTool.generated.h"
+
+class AGroomActor;
+class AStaticMeshActor;
+PREDECLARE_USE_GEOMETRY_CLASS(FDynamicMesh3);
+
+/**
+ *
+ */
+UCLASS()
+class HAIRMODELINGTOOLSET_API UGroomToMeshToolBuilder : public UInteractiveToolBuilder
+{
+	GENERATED_BODY()
+public:
+	virtual bool CanBuildTool(const FToolBuilderState& SceneState) const override;
+	virtual UInteractiveTool* BuildTool(const FToolBuilderState& SceneState) const override;
+};
+
+
+
+
+UENUM()
+enum class EGroomToMeshUVMode
+{
+	PlanarSplitting = 1,
+	MinimalConformal = 2,
+	PlanarSplitConformal = 3
+};
+
+
+
+UCLASS()
+class HAIRMODELINGTOOLSET_API UGroomToMeshToolProperties : public UInteractiveToolPropertySet
+{
+	GENERATED_BODY()
+
+public:
+	/** The size of the geometry bounding box major axis measured in voxels */
+	UPROPERTY(EditAnywhere, Category = Meshing, meta = (UIMin = "8", UIMax = "512", ClampMin = "8", ClampMax = "1024"))
+	int32 VoxelCount = 64;
+
+	UPROPERTY(EditAnywhere, Category = Meshing, meta = (UIMin = "0.5", UIMax = "16.0", ClampMin = "0.1", ClampMax = "128.0"))
+	float BlendPower = 1.0;
+
+	UPROPERTY(EditAnywhere, Category = Meshing, meta = (UIMin = "0.1", UIMax = "4.0", ClampMin = "0.1", ClampMax = "128.0"))
+	float RadiusScale = 0.5;
+
+
+	UPROPERTY(EditAnywhere, Category = Morphology)
+	bool bApplyMorphology = true;
+
+	UPROPERTY(EditAnywhere, Category = Morphology, meta = (UIMin = "0.0", UIMax = "50.0", ClampMin = "0.0", ClampMax = "128.0"))
+	float ClosingDist = 2.0;
+
+	UPROPERTY(EditAnywhere, Category = Morphology, meta = (UIMin = "0.0", UIMax = "50.0", ClampMin = "0.0", ClampMax = "128.0"))
+	float OpeningDist = 0.25;
+
+
+
+	UPROPERTY(EditAnywhere, Category = Clipping)
+	bool bClipToHead = true;
+
+	// todo: this probably also needs to support skeletal mesh
+	UPROPERTY(EditAnywhere, Category = Clipping)
+	TLazyObjectPtr<AStaticMeshActor> ClipMeshActor;
+
+
+	UPROPERTY(EditAnywhere, Category = Smoothing)
+	bool bSmooth = true;
+
+	UPROPERTY(EditAnywhere, Category = Smoothing, meta = (UIMin = "0.0", UIMax = "1.0", ClampMin = "0.0", ClampMax = "1.0"))
+	float Smoothness = 0.15;
+
+	UPROPERTY(EditAnywhere, Category = Smoothing, meta = (UIMin = "-1.0", UIMax = "1.0", ClampMin = "-2.0", ClampMax = "2.0"))
+	float VolumeCorrection = -0.25f;
+
+
+	UPROPERTY(EditAnywhere, Category = Simplification)
+	bool bSimplify = false;
+
+	/** Target triangle count */
+	UPROPERTY(EditAnywhere, Category = Simplification, meta = (UIMin = "4", UIMax = "5000", ClampMin = "1", ClampMax = "9999999999", EditCondition = "bSimplify == true"))
+	int VertexCount = 500;
+
+
+	UPROPERTY(EditAnywhere, Category = UVGeneration)
+	EGroomToMeshUVMode UVMode = EGroomToMeshUVMode::MinimalConformal;
+
+
+	UPROPERTY(EditAnywhere, Category = Visualization)
+	bool bShowSideBySide = true;
+
+	UPROPERTY(EditAnywhere, Category = Visualization)
+	bool bShowGuides = false;
+
+	UPROPERTY(EditAnywhere, Category = Visualization)
+	bool bShowUVs = false;
+
+};
+
+
+
+/**
+ *
+ */
+UCLASS()
+class HAIRMODELINGTOOLSET_API UGroomToMeshTool : public UInteractiveTool
+{
+	GENERATED_BODY()
+
+public:
+	UGroomToMeshTool();
+
+	virtual void SetWorld(UWorld* World) { this->TargetWorld = World; }
+	virtual void SetSelection(AGroomActor* Groom);
+
+	virtual void Setup() override;
+	virtual void Shutdown(EToolShutdownType ShutdownType) override;
+
+	virtual void Render(IToolsContextRenderAPI* RenderAPI) override;
+	virtual void OnTick(float DeltaTime) override;
+
+	virtual bool HasCancel() const override { return true; }
+	virtual bool HasAccept() const override;
+	virtual bool CanAccept() const override;
+
+protected:
+	UPROPERTY()
+	TObjectPtr<UGroomToMeshToolProperties> Settings;
+
+	UPROPERTY()
+	TObjectPtr<UPreviewMesh> PreviewMesh;
+
+	UPROPERTY()
+	TLazyObjectPtr<AGroomActor> TargetGroom;
+
+	UPROPERTY()
+	TObjectPtr<UPreviewGeometry> PreviewGeom;
+
+protected:
+	UPROPERTY()
+	TObjectPtr<UMaterialInterface> MeshMaterial;
+
+	UPROPERTY()
+	TObjectPtr<UMaterialInterface> UVMaterial;
+
+protected:
+	UWorld* TargetWorld = nullptr;
+
+	FDynamicMesh3 CurrentMesh;
+
+	void RecalculateMesh();
+
+	void UpdateLineSet();
+
+	bool bResultValid = false;
+	bool bVisualizationChanged = false;
+
+	struct FVoxelizeSettings
+	{
+		int32 VoxelCount;
+		float BlendPower;
+		float RadiusScale;
+		
+		bool operator==(const FVoxelizeSettings& Other) const
+		{
+			return VoxelCount == Other.VoxelCount && BlendPower == Other.BlendPower && RadiusScale == Other.RadiusScale;
+		}
+	};
+	FVoxelizeSettings CachedVoxelizeSettings;
+	TSharedPtr<FDynamicMesh3, ESPMode::ThreadSafe> CurrentVoxelizeResult;
+	TSharedPtr<FDynamicMesh3, ESPMode::ThreadSafe> UpdateVoxelization();
+
+
+
+	struct FMorphologySettings
+	{
+		TSharedPtr<FDynamicMesh3, ESPMode::ThreadSafe> InputMesh;
+		int32 VoxelCount;
+		float CloseDist;
+		float OpenDist;
+
+		bool operator==(const FMorphologySettings& Other) const
+		{
+			return InputMesh.Get() == Other.InputMesh.Get() && VoxelCount == Other.VoxelCount && OpenDist == Other.OpenDist && CloseDist == Other.CloseDist;
+		}
+	};
+	FMorphologySettings CachedMorphologySettings;
+	TSharedPtr<FDynamicMesh3, ESPMode::ThreadSafe> CachedMorphologyResult;
+	TSharedPtr<FDynamicMesh3, ESPMode::ThreadSafe> UpdateMorphology(TSharedPtr<FDynamicMesh3, ESPMode::ThreadSafe> InputMesh);
+
+
+	struct FClipMeshSettings
+	{
+		TSharedPtr<FDynamicMesh3, ESPMode::ThreadSafe> InputMesh;
+		AActor* ClipSource;
+
+		bool operator==(const FClipMeshSettings& Other) const
+		{
+			return InputMesh.Get() == Other.InputMesh.Get() && ClipSource == Other.ClipSource;
+		}
+	};
+	FClipMeshSettings CachedClipMeshSettings;
+	TSharedPtr<FDynamicMesh3, ESPMode::ThreadSafe> CachedClipMeshResult;
+	TSharedPtr<FDynamicMesh3, ESPMode::ThreadSafe> UpdateClipMesh(TSharedPtr<FDynamicMesh3, ESPMode::ThreadSafe> InputMesh);
+
+
+	struct FSmoothingSettings
+	{
+		TSharedPtr<FDynamicMesh3, ESPMode::ThreadSafe> InputMesh;
+		float Smoothness;
+		float VolumeCorrection;
+
+		bool operator==(const FSmoothingSettings& Other) const
+		{
+			return InputMesh.Get() == Other.InputMesh.Get() && Smoothness == Other.Smoothness && VolumeCorrection == Other.VolumeCorrection;
+		}
+	};
+	FSmoothingSettings CachedSmoothSettings;
+	TSharedPtr<FDynamicMesh3, ESPMode::ThreadSafe> CachedSmoothResult;
+	TSharedPtr<FDynamicMesh3, ESPMode::ThreadSafe> UpdateSmoothing(TSharedPtr<FDynamicMesh3, ESPMode::ThreadSafe> InputMesh);
+
+
+	struct FSimplifySettings
+	{
+		TSharedPtr<FDynamicMesh3, ESPMode::ThreadSafe> InputMesh;
+		int32 TargetCount;
+
+		bool operator==(const FSimplifySettings& Other) const
+		{
+			return TargetCount == Other.TargetCount && InputMesh.Get() == Other.InputMesh.Get();
+		}
+	};
+	FSimplifySettings CachedSimplifySettings;
+	TSharedPtr<FDynamicMesh3, ESPMode::ThreadSafe> CachedSimplifyResult;
+	TSharedPtr<FDynamicMesh3, ESPMode::ThreadSafe> UpdateSimplification(TSharedPtr<FDynamicMesh3, ESPMode::ThreadSafe> InputMesh);
+
+
+	struct FPostprocessSettings
+	{
+		TSharedPtr<FDynamicMesh3, ESPMode::ThreadSafe> InputMesh;
+		EGroomToMeshUVMode UVGenMode;
+
+		bool operator==(const FPostprocessSettings& Other) const
+		{
+			return InputMesh.Get() == Other.InputMesh.Get() && UVGenMode == Other.UVGenMode;
+		}
+	};
+	FPostprocessSettings CachedPostprocessSettings;
+	TSharedPtr<FDynamicMesh3, ESPMode::ThreadSafe> CachedPostprocessResult;
+	TSharedPtr<FDynamicMesh3, ESPMode::ThreadSafe> UpdatePostprocessing(TSharedPtr<FDynamicMesh3, ESPMode::ThreadSafe> InputMesh);
+
+
+	TSharedPtr<FDynamicMesh3, ESPMode::ThreadSafe> UpdateUVs(TSharedPtr<FDynamicMesh3, ESPMode::ThreadSafe> InputMesh, EGroomToMeshUVMode UVMode);
+
+	TSharedPtr<FDynamicMesh3, ESPMode::ThreadSafe> UpdateUVs_ExpMapPlaneSplits(TSharedPtr<FDynamicMesh3, ESPMode::ThreadSafe> InputMesh, bool bRecalcAsConformal);
+	TSharedPtr<FDynamicMesh3, ESPMode::ThreadSafe> UpdateUVs_MinimalConformal(TSharedPtr<FDynamicMesh3, ESPMode::ThreadSafe> InputMesh);
+
+	void UpdatePreview(TSharedPtr<FDynamicMesh3, ESPMode::ThreadSafe> ResultMesh);
+};
diff --git a/Engine/Plugins/Experimental/HairModelingToolset/Source/HairModelingToolset/Public/HairModelingToolCommands.h b/Engine/Plugins/Experimental/HairModelingToolset/Source/HairModelingToolset/Public/HairModelingToolCommands.h
new file mode 100644
--- /dev/null
+++ b/Engine/Plugins/Experimental/HairModelingToolset/Source/HairModelingToolset/Public/HairModelingToolCommands.h
@@ -0,0 +1,24 @@
+// Copyright Epic Games, Inc. All Rights Reserved.
+
+#pragma once
+
+#include "CoreMinimal.h"
+#include "Framework/Commands/Commands.h"
+
+/**
+ * TInteractiveToolCommands implementation for this module that provides standard Editor hotkey support
+ */
+class HAIRMODELINGTOOLSET_API FHairModelingToolCommands : public TCommands<FHairModelingToolCommands>
+{
+public:
+	FHairModelingToolCommands();
+
+	TSharedPtr<FUICommandInfo> BeginGroomToMeshTool;
+	TSharedPtr<FUICommandInfo> BeginGenerateLODMeshesTool;
+	TSharedPtr<FUICommandInfo> BeginGroomCardsEditorTool;
+
+	/**
+	 * Initialize commands
+	 */
+	virtual void RegisterCommands() override;
+};
diff --git a/Engine/Plugins/Experimental/HairModelingToolset/Source/HairModelingToolset/Public/HairModelingToolsStyle.h b/Engine/Plugins/Experimental/HairModelingToolset/Source/HairModelingToolset/Public/HairModelingToolsStyle.h
new file mode 100644
--- /dev/null
+++ b/Engine/Plugins/Experimental/HairModelingToolset/Source/HairModelingToolset/Public/HairModelingToolsStyle.h
@@ -0,0 +1,24 @@
+// Copyright Epic Games, Inc. All Rights Reserved.
+
+#pragma once
+
+#include "CoreMinimal.h"
+#include "Styling/ISlateStyle.h"
+#include "Styling/SlateStyle.h"
+
+class HAIRMODELINGTOOLSET_API FHairModelingToolsStyle
+{
+public:
+	static void Initialize();
+
+	static void Shutdown();
+
+	static TSharedPtr< class ISlateStyle > Get();
+
+	static FName GetStyleSetName();
+private:
+	static FString InContent(const FString& RelativePath, const ANSICHAR* Extension);
+
+private:
+	static TSharedPtr< class FSlateStyleSet > StyleSet;
+};
diff --git a/Engine/Plugins/Experimental/HairModelingToolset/Source/HairModelingToolset/Public/HairModelingToolsetModule.h b/Engine/Plugins/Experimental/HairModelingToolset/Source/HairModelingToolset/Public/HairModelingToolsetModule.h
new file mode 100644
--- /dev/null
+++ b/Engine/Plugins/Experimental/HairModelingToolset/Source/HairModelingToolset/Public/HairModelingToolsetModule.h
@@ -0,0 +1,21 @@
+// Copyright Epic Games, Inc. All Rights Reserved.
+
+#pragma once
+
+#include "CoreMinimal.h"
+#include "Modules/ModuleManager.h"
+#include "ModelingModeToolExtensions.h"
+
+class FHairModelingToolsetModule : public IModuleInterface, public IModelingModeToolExtension
+{
+public:
+
+	/** IModuleInterface implementation */
+	virtual void StartupModule() override;
+	virtual void ShutdownModule() override;
+
+	/** IModelingModeToolExtension implementation */
+	virtual FText GetExtensionName() override;
+	virtual FText GetToolSectionName() override;
+	virtual void GetExtensionTools(const FExtensionToolQueryInfo& QueryInfo, TArray<FExtensionToolDescription>& ToolsOut) override;
+};
diff --git a/Engine/Plugins/Experimental/MeshModelingToolsetExp/Source/MeshModelingToolsEditorOnly/MeshModelingToolsEditorOnly.Build.cs b/Engine/Plugins/Experimental/MeshModelingToolsetExp/Source/MeshModelingToolsEditorOnly/MeshModelingToolsEditorOnly.Build.cs
--- a/Engine/Plugins/Experimental/MeshModelingToolsetExp/Source/MeshModelingToolsEditorOnly/MeshModelingToolsEditorOnly.Build.cs
+++ b/Engine/Plugins/Experimental/MeshModelingToolsetExp/Source/MeshModelingToolsEditorOnly/MeshModelingToolsEditorOnly.Build.cs
@@ -40,8 +40,6 @@ public class MeshModelingToolsEditorOnly : ModuleRules
 				"ModelingOperators",
 				"ModelingOperatorsEditorOnly",
 				// ... add other public dependencies that you statically link with here ...
-
-				"HairStrandsCore"		// required for Hair toolset
 			}
             );
 			
diff --git a/Engine/Plugins/Experimental/MeshModelingToolsetExp/Source/MeshModelingToolsEditorOnly/Private/GenerateLODMeshesTool.cpp b/Engine/Plugins/Experimental/MeshModelingToolsetExp/Source/MeshModelingToolsEditorOnly/Private/GenerateLODMeshesTool.cpp
deleted file mode 100644
diff --git a/Engine/Plugins/Experimental/MeshModelingToolsetExp/Source/MeshModelingToolsEditorOnly/Private/Hair/GroomCardsEditorTool.cpp b/Engine/Plugins/Experimental/MeshModelingToolsetExp/Source/MeshModelingToolsEditorOnly/Private/Hair/GroomCardsEditorTool.cpp
deleted file mode 100644
diff --git a/Engine/Plugins/Experimental/MeshModelingToolsetExp/Source/MeshModelingToolsEditorOnly/Private/Hair/GroomQueryUtil.cpp b/Engine/Plugins/Experimental/MeshModelingToolsetExp/Source/MeshModelingToolsEditorOnly/Private/Hair/GroomQueryUtil.cpp
deleted file mode 100644
diff --git a/Engine/Plugins/Experimental/MeshModelingToolsetExp/Source/MeshModelingToolsEditorOnly/Private/Hair/GroomQueryUtil.h b/Engine/Plugins/Experimental/MeshModelingToolsetExp/Source/MeshModelingToolsEditorOnly/Private/Hair/GroomQueryUtil.h
deleted file mode 100644
diff --git a/Engine/Plugins/Experimental/MeshModelingToolsetExp/Source/MeshModelingToolsEditorOnly/Private/Hair/GroomToMeshTool.cpp b/Engine/Plugins/Experimental/MeshModelingToolsetExp/Source/MeshModelingToolsEditorOnly/Private/Hair/GroomToMeshTool.cpp
deleted file mode 100644
diff --git a/Engine/Plugins/Experimental/MeshModelingToolsetExp/Source/MeshModelingToolsEditorOnly/Public/GenerateLODMeshesTool.h b/Engine/Plugins/Experimental/MeshModelingToolsetExp/Source/MeshModelingToolsEditorOnly/Public/GenerateLODMeshesTool.h
deleted file mode 100644
diff --git a/Engine/Plugins/Experimental/MeshModelingToolsetExp/Source/MeshModelingToolsEditorOnly/Public/Hair/GroomCardsEditorTool.h b/Engine/Plugins/Experimental/MeshModelingToolsetExp/Source/MeshModelingToolsEditorOnly/Public/Hair/GroomCardsEditorTool.h
deleted file mode 100644
diff --git a/Engine/Plugins/Experimental/MeshModelingToolsetExp/Source/MeshModelingToolsEditorOnly/Public/Hair/GroomToMeshTool.h b/Engine/Plugins/Experimental/MeshModelingToolsetExp/Source/MeshModelingToolsEditorOnly/Public/Hair/GroomToMeshTool.h
deleted file mode 100644
