More clean up & sort #include paths in various files

diff --git a/Engine/Plugins/Editor/WaveformEditor/Source/WaveformTransformations/Public/WaveformTransformationEffectChain.h b/Engine/Plugins/Editor/WaveformEditor/Source/WaveformTransformations/Public/WaveformTransformationEffectChain.h
--- a/Engine/Plugins/Editor/WaveformEditor/Source/WaveformTransformations/Public/WaveformTransformationEffectChain.h
+++ b/Engine/Plugins/Editor/WaveformEditor/Source/WaveformTransformations/Public/WaveformTransformationEffectChain.h
@@ -2,8 +2,8 @@
 
 #pragma once
 
-#include "AudioExtensions/Public/IWaveformTransformation.h"
 #include "Factories/SoundSourceEffectFactory.h"
+#include "IWaveformTransformation.h"
 #include "UObject/Object.h"
 #include "UObject/StrongObjectPtr.h"
 #include "WaveformTransformationEffectChain.generated.h"
diff --git a/Engine/Plugins/Editor/WaveformEditor/Source/WaveformTransformations/Public/WaveformTransformationNormalize.h b/Engine/Plugins/Editor/WaveformEditor/Source/WaveformTransformations/Public/WaveformTransformationNormalize.h
--- a/Engine/Plugins/Editor/WaveformEditor/Source/WaveformTransformations/Public/WaveformTransformationNormalize.h
+++ b/Engine/Plugins/Editor/WaveformEditor/Source/WaveformTransformations/Public/WaveformTransformationNormalize.h
@@ -2,7 +2,7 @@
 
 #pragma once
 
-#include "AudioExtensions/Public/IWaveformTransformation.h"
+#include "IWaveformTransformation.h"
 #include "UObject/Object.h"
 #include "WaveformTransformationNormalize.generated.h"
 
diff --git a/Engine/Plugins/Enterprise/DatasmithFBXImporter/Source/DatasmithFBXTranslator/Public/DatasmithFBXScene.h b/Engine/Plugins/Enterprise/DatasmithFBXImporter/Source/DatasmithFBXTranslator/Public/DatasmithFBXScene.h
--- a/Engine/Plugins/Enterprise/DatasmithFBXImporter/Source/DatasmithFBXTranslator/Public/DatasmithFBXScene.h
+++ b/Engine/Plugins/Enterprise/DatasmithFBXImporter/Source/DatasmithFBXTranslator/Public/DatasmithFBXScene.h
@@ -9,7 +9,7 @@
 #include "CoreTypes.h"
 #include "Curves/RealCurve.h"
 #include "Curves/RichCurve.h"
-#include "Engine/Classes/EditorFramework/AssetImportData.h"
+#include "EditorFramework/AssetImportData.h"
 #include "Engine/DataTable.h" //For FTableRowBase
 #include "MeshDescription.h"
 #include "Misc/SecureHash.h"
diff --git a/Engine/Plugins/Experimental/HairModelingToolset/Source/HairModelingToolset/Private/GroomCardsEditorTool.cpp b/Engine/Plugins/Experimental/HairModelingToolset/Source/HairModelingToolset/Private/GroomCardsEditorTool.cpp
--- a/Engine/Plugins/Experimental/HairModelingToolset/Source/HairModelingToolset/Private/GroomCardsEditorTool.cpp
+++ b/Engine/Plugins/Experimental/HairModelingToolset/Source/HairModelingToolset/Private/GroomCardsEditorTool.cpp
@@ -27,9 +27,9 @@
 
 #include "GroomQueryUtil.h"
 
-#include "Engine/Classes/Engine/StaticMesh.h"
-#include "Engine/Classes/Components/StaticMeshComponent.h"
-#include "Engine/Classes/Engine/StaticMeshActor.h"
+#include "Components/StaticMeshComponent.h"
+#include "Engine/StaticMesh.h"
+#include "Engine/StaticMeshActor.h"
 
 
 #define LOCTEXT_NAMESPACE "UGroomCardsEditorTool"
diff --git a/Engine/Plugins/Experimental/HairModelingToolset/Source/HairModelingToolset/Private/GroomToMeshTool.cpp b/Engine/Plugins/Experimental/HairModelingToolset/Source/HairModelingToolset/Private/GroomToMeshTool.cpp
--- a/Engine/Plugins/Experimental/HairModelingToolset/Source/HairModelingToolset/Private/GroomToMeshTool.cpp
+++ b/Engine/Plugins/Experimental/HairModelingToolset/Source/HairModelingToolset/Private/GroomToMeshTool.cpp
@@ -44,9 +44,9 @@
 
 #include "ParameterizationOps/UVLayoutOp.h"
 
-#include "Engine/Classes/Engine/StaticMesh.h"
-#include "Engine/Classes/Components/StaticMeshComponent.h"
-#include "Engine/Classes/Engine/StaticMeshActor.h"
+#include "Components/StaticMeshComponent.h"
+#include "Engine/StaticMesh.h"
+#include "Engine/StaticMeshActor.h"
 #include "MeshDescriptionToDynamicMesh.h"
 
 using namespace UE::Geometry;
diff --git a/Engine/Plugins/Runtime/nDisplay/Source/DisplayCluster/Private/Render/Synchronization/Windows/DisplayClusterRenderSyncPolicyEthernetWindows.cpp b/Engine/Plugins/Runtime/nDisplay/Source/DisplayCluster/Private/Render/Synchronization/Windows/DisplayClusterRenderSyncPolicyEthernetWindows.cpp
--- a/Engine/Plugins/Runtime/nDisplay/Source/DisplayCluster/Private/Render/Synchronization/Windows/DisplayClusterRenderSyncPolicyEthernetWindows.cpp
+++ b/Engine/Plugins/Runtime/nDisplay/Source/DisplayCluster/Private/Render/Synchronization/Windows/DisplayClusterRenderSyncPolicyEthernetWindows.cpp
@@ -13,7 +13,7 @@
 #include "UnrealClient.h"
 #include "Windows/AllowWindowsPlatformTypes.h"
 #include <windows.h>
-#include "DirectX/Include/DXGI.h"
+#include "DXGI.h"
 #include "dwmapi.h"
 #include "Windows/HideWindowsPlatformTypes.h"
 
