Fix Hair Modeling toolset pluging compilation error.

diff --git a/Engine/Plugins/Experimental/HairModelingToolset/Source/HairModelingToolset/Private/GroomQueryUtil.cpp b/Engine/Plugins/Experimental/HairModelingToolset/Source/HairModelingToolset/Private/GroomQueryUtil.cpp
--- a/Engine/Plugins/Experimental/HairModelingToolset/Source/HairModelingToolset/Private/GroomQueryUtil.cpp
+++ b/Engine/Plugins/Experimental/HairModelingToolset/Source/HairModelingToolset/Private/GroomQueryUtil.cpp
@@ -36,11 +36,11 @@ void UE::GroomQueries::ExtractAllHairCards(AGroomActor* GroomActor,
 	FDynamicMeshUVOverlay* UVs = MeshOut.Attributes()->PrimaryUV();
 	FDynamicMeshNormalOverlay* Normals = MeshOut.Attributes()->PrimaryNormals();
 
-	int32 NumHairGroups = Asset->HairGroupsInfo.Num();
+	int32 NumHairGroups = Asset->GetHairGroupsInfo().Num();
 	for (int32 GroupIdx = 0; GroupIdx < NumHairGroups; ++GroupIdx)
 	{
-		const FHairGroupInfo& GroupInfo = Asset->HairGroupsInfo[GroupIdx];
-		const FHairGroupPlatformData& GroupData = Asset->HairGroupsPlatformData[GroupIdx];
+		const FHairGroupInfo& GroupInfo = Asset->GetHairGroupsInfo()[GroupIdx];
+		const FHairGroupPlatformData& GroupData = Asset->GetHairGroupsPlatformData()[GroupIdx];
 
 		int32 NumCardsLODs = GroupData.Cards.LODs.Num();
 		if (LODIndex >= NumCardsLODs)
@@ -49,7 +49,7 @@ void UE::GroomQueries::ExtractAllHairCards(AGroomActor* GroomActor,
 		}
 
 		UStaticMesh* StaticMesh = nullptr;
-		for (const FHairGroupsCardsSourceDescription& Desc : Asset->HairGroupsCards)
+		for (const FHairGroupsCardsSourceDescription& Desc : Asset->GetHairGroupsCards())
 		{
 			if (Desc.GroupIndex == GroupIdx && Desc.LODIndex == LODIndex)
 			{
@@ -516,11 +516,11 @@ void UE::GroomQueries::ProcessHairCurves(AGroomActor* GroomActor,
 	check(GroomActor->GetGroomComponent()->GroomAsset);
 	UGroomAsset* Asset = GroomActor->GetGroomComponent()->GroomAsset;
 
-	int32 NumHairGroups = Asset->HairGroupsInfo.Num();
+	int32 NumHairGroups = Asset->GetHairGroupsInfo().Num();
 	for (int32 GroupIdx = 0; GroupIdx < NumHairGroups; ++GroupIdx)
 	{
-		const FHairGroupInfo& GroupInfo = Asset->HairGroupsInfo[GroupIdx];
-		const FHairGroupPlatformData& GroupData = Asset->HairGroupsPlatformData[GroupIdx];
+		const FHairGroupInfo& GroupInfo = Asset->GetHairGroupsInfo()[GroupIdx];
+		const FHairGroupPlatformData& GroupData = Asset->GetHairGroupsPlatformData()[GroupIdx];
 
 		//int32 NumCurves = GroupInfo.NumCurves;
 		//int32 NumGuides = GroupInfo.NumGuides;
diff --git a/Engine/Plugins/Experimental/HairModelingToolset/Source/HairModelingToolset/Private/GroomToMeshTool.cpp b/Engine/Plugins/Experimental/HairModelingToolset/Source/HairModelingToolset/Private/GroomToMeshTool.cpp
--- a/Engine/Plugins/Experimental/HairModelingToolset/Source/HairModelingToolset/Private/GroomToMeshTool.cpp
+++ b/Engine/Plugins/Experimental/HairModelingToolset/Source/HairModelingToolset/Private/GroomToMeshTool.cpp
@@ -297,11 +297,11 @@ static void ProcessHairCurvePoints(AGroomActor* GroomActor,
 	check(GroomActor->GetGroomComponent()->GroomAsset);
 	UGroomAsset* Asset = GroomActor->GetGroomComponent()->GroomAsset;
 
-	int32 NumHairGroups = Asset->HairGroupsInfo.Num();
+	int32 NumHairGroups = Asset->GetHairGroupsInfo().Num();
 	for (int32 GroupIdx = 0; GroupIdx < NumHairGroups; ++GroupIdx)
 	{
-		const FHairGroupInfo& GroupInfo = Asset->HairGroupsInfo[GroupIdx];
-		const FHairGroupPlatformData& GroupData = Asset->HairGroupsPlatformData[GroupIdx];
+		const FHairGroupInfo& GroupInfo = Asset->GetHairGroupsInfo()[GroupIdx];
+		const FHairGroupPlatformData& GroupData = Asset->GetHairGroupsPlatformData()[GroupIdx];
 
 		//int32 NumCurves = GroupInfo.NumCurves;
 		//int32 NumGuides = GroupInfo.NumGuides;
@@ -343,11 +343,11 @@ static void ProcessHairCurves(AGroomActor* GroomActor,
 	check(GroomActor->GetGroomComponent()->GroomAsset);
 	UGroomAsset* Asset = GroomActor->GetGroomComponent()->GroomAsset;
 
-	int32 NumHairGroups = Asset->HairGroupsInfo.Num();
+	int32 NumHairGroups = Asset->GetHairGroupsInfo().Num();
 	for (int32 GroupIdx = 0; GroupIdx < NumHairGroups; ++GroupIdx)
 	{
-		const FHairGroupInfo& GroupInfo = Asset->HairGroupsInfo[GroupIdx];
-		const FHairGroupPlatformData& GroupData = Asset->HairGroupsPlatformData[GroupIdx];
+		const FHairGroupInfo& GroupInfo = Asset->GetHairGroupsInfo()[GroupIdx];
+		const FHairGroupPlatformData& GroupData = Asset->GetHairGroupsPlatformData()[GroupIdx];
 
 		//int32 NumCurves = GroupInfo.NumCurves;
 		//int32 NumGuides = GroupInfo.NumGuides;
