GroomToMesh tool: update the default value for RadiusScale to account for hair radii values no longer being normalized

diff --git a/Engine/Plugins/Experimental/HairModelingToolset/Source/HairModelingToolset/Public/GroomToMeshTool.h b/Engine/Plugins/Experimental/HairModelingToolset/Source/HairModelingToolset/Public/GroomToMeshTool.h
--- a/Engine/Plugins/Experimental/HairModelingToolset/Source/HairModelingToolset/Public/GroomToMeshTool.h
+++ b/Engine/Plugins/Experimental/HairModelingToolset/Source/HairModelingToolset/Public/GroomToMeshTool.h
@@ -53,8 +53,8 @@ public:
 	UPROPERTY(EditAnywhere, Category = Meshing, meta = (UIMin = "0.5", UIMax = "16.0", ClampMin = "0.1", ClampMax = "128.0"))
 	float BlendPower = 1.0;
 
-	UPROPERTY(EditAnywhere, Category = Meshing, meta = (UIMin = "0.1", UIMax = "4.0", ClampMin = "0.1", ClampMax = "128.0"))
-	float RadiusScale = 0.5;
+	UPROPERTY(EditAnywhere, Category = Meshing, meta = (UIMin = "0.1", UIMax = "200.0", ClampMin = "0.1", ClampMax = "2000.0"))
+	float RadiusScale = 100.0;
 
 
 	UPROPERTY(EditAnywhere, Category = Morphology)
