TLazyObjectPtr Deprecation pt 2.:
Change usage of TLazyObjectPtr to TSoftObjectPtr in UInteractiveTool usage.
Those classes are transient and do not persist

#Synced-CL 52252558

diff --git a/Engine/Plugins/Experimental/HairModelingToolset/Source/HairModelingToolset/Private/GroomToMeshTool.cpp b/Engine/Plugins/Experimental/HairModelingToolset/Source/HairModelingToolset/Private/GroomToMeshTool.cpp
--- a/Engine/Plugins/Experimental/HairModelingToolset/Source/HairModelingToolset/Private/GroomToMeshTool.cpp
+++ b/Engine/Plugins/Experimental/HairModelingToolset/Source/HairModelingToolset/Private/GroomToMeshTool.cpp
@@ -124,7 +124,7 @@ void UGroomToMeshTool::Setup()
 	Settings->WatchProperty(Settings->OpeningDist, [this](float) { bResultValid = false; });
 
 	Settings->WatchProperty(Settings->bClipToHead, [this](bool) { bResultValid = false; });
-	Settings->WatchProperty(Settings->ClipMeshActor, [this](TLazyObjectPtr<AStaticMeshActor>) { bResultValid = false; });
+	Settings->WatchProperty(Settings->ClipMeshActor, [this](TSoftObjectPtr<AStaticMeshActor>) { bResultValid = false; });
 
 	Settings->WatchProperty(Settings->bSmooth, [this](bool) { bResultValid = false; });
 	Settings->WatchProperty(Settings->Smoothness, [this](float) { bResultValid = false; });
diff --git a/Engine/Plugins/Experimental/HairModelingToolset/Source/HairModelingToolset/Public/GroomCardsEditorTool.h b/Engine/Plugins/Experimental/HairModelingToolset/Source/HairModelingToolset/Public/GroomCardsEditorTool.h
--- a/Engine/Plugins/Experimental/HairModelingToolset/Source/HairModelingToolset/Public/GroomCardsEditorTool.h
+++ b/Engine/Plugins/Experimental/HairModelingToolset/Source/HairModelingToolset/Public/GroomCardsEditorTool.h
@@ -197,7 +197,7 @@ protected:
 	TObjectPtr<UPreviewMesh> PreviewMesh;
 
 	UPROPERTY()
-	TLazyObjectPtr<AGroomActor> TargetGroom;
+	TSoftObjectPtr<AGroomActor> TargetGroom;
 
 	UPROPERTY()
 	TObjectPtr<UPreviewGeometry> PreviewGeom;
diff --git a/Engine/Plugins/Experimental/HairModelingToolset/Source/HairModelingToolset/Public/GroomToMeshTool.h b/Engine/Plugins/Experimental/HairModelingToolset/Source/HairModelingToolset/Public/GroomToMeshTool.h
--- a/Engine/Plugins/Experimental/HairModelingToolset/Source/HairModelingToolset/Public/GroomToMeshTool.h
+++ b/Engine/Plugins/Experimental/HairModelingToolset/Source/HairModelingToolset/Public/GroomToMeshTool.h
@@ -73,7 +73,7 @@ public:
 
 	// todo: this probably also needs to support skeletal mesh
 	UPROPERTY(EditAnywhere, Category = Clipping)
-	TLazyObjectPtr<AStaticMeshActor> ClipMeshActor;
+	TSoftObjectPtr<AStaticMeshActor> ClipMeshActor;
 
 
 	UPROPERTY(EditAnywhere, Category = Smoothing)
@@ -143,7 +143,7 @@ protected:
 	TObjectPtr<UPreviewMesh> PreviewMesh;
 
 	UPROPERTY()
-	TLazyObjectPtr<AGroomActor> TargetGroom;
+	TSoftObjectPtr<AGroomActor> TargetGroom;
 
 	UPROPERTY()
 	TObjectPtr<UPreviewGeometry> PreviewGeom;
diff --git a/Engine/Plugins/Experimental/MeshModelingToolsetExp/Source/MeshModelingToolsEditorOnlyExp/Private/MeshToVolumeTool.cpp b/Engine/Plugins/Experimental/MeshModelingToolsetExp/Source/MeshModelingToolsEditorOnlyExp/Private/MeshToVolumeTool.cpp
--- a/Engine/Plugins/Experimental/MeshModelingToolsetExp/Source/MeshModelingToolsEditorOnlyExp/Private/MeshToVolumeTool.cpp
+++ b/Engine/Plugins/Experimental/MeshModelingToolsetExp/Source/MeshModelingToolsEditorOnlyExp/Private/MeshToVolumeTool.cpp
@@ -212,7 +212,8 @@ void UMeshToVolumeTool::OnShutdown(EToolShutdownType ShutdownType)
 
 		AVolume* TargetVolume = nullptr;
 
-		if (Settings->TargetVolume.IsValid() == false)
+		const bool bHasTargetVolume = Settings->TargetVolume.IsValid();
+		if (bHasTargetVolume == false)
 		{
 			FRotator Rotation(0.0f, 0.0f, 0.0f);
 			FActorSpawnParameters SpawnInfo;
diff --git a/Engine/Plugins/Experimental/MeshModelingToolsetExp/Source/MeshModelingToolsEditorOnlyExp/Public/MeshToVolumeTool.h b/Engine/Plugins/Experimental/MeshModelingToolsetExp/Source/MeshModelingToolsEditorOnlyExp/Public/MeshToVolumeTool.h
--- a/Engine/Plugins/Experimental/MeshModelingToolsetExp/Source/MeshModelingToolsEditorOnlyExp/Public/MeshToVolumeTool.h
+++ b/Engine/Plugins/Experimental/MeshModelingToolsetExp/Source/MeshModelingToolsEditorOnlyExp/Public/MeshToVolumeTool.h
@@ -81,7 +81,8 @@ public:
 
 	/** If set, the target Volume will be updated, rather than creating a new Volume. */
 	UPROPERTY(EditAnywhere, Category = UpdateExisting)
-	TLazyObjectPtr<AVolume> TargetVolume;
+	TSoftObjectPtr<AVolume> TargetVolume;
+
 };
 
 
diff --git a/Engine/Plugins/Experimental/MeshModelingToolsetExp/Source/MeshModelingToolsExp/Public/VolumeToMeshTool.h b/Engine/Plugins/Experimental/MeshModelingToolsetExp/Source/MeshModelingToolsExp/Public/VolumeToMeshTool.h
--- a/Engine/Plugins/Experimental/MeshModelingToolsetExp/Source/MeshModelingToolsExp/Public/VolumeToMeshTool.h
+++ b/Engine/Plugins/Experimental/MeshModelingToolsetExp/Source/MeshModelingToolsExp/Public/VolumeToMeshTool.h
@@ -91,7 +91,7 @@ protected:
 	TObjectPtr<UPreviewMesh> PreviewMesh;
 
 	UPROPERTY()
-	TLazyObjectPtr<AVolume> TargetVolume;
+	TSoftObjectPtr<AVolume> TargetVolume;
 
 	UPROPERTY()
 	TObjectPtr<ULineSetComponent> VolumeEdgesSet;
